Created a new interface to RSA functionality that is independent of key size (unlike the old TBigInt version). Internally, RSA functions are implemented by OpenSSL on Windows/Mac/Linux an falling back to the TBigInt implementation for other platforms
- Future plan would be not to use OpenSSL directly, but to go via PlatformCrypto or something like it which could provide platform specific RSA functionality if it exists
- Also contains some compensation code for platforms that are still stuck on an older version of OpenSSL (everything except windows)
- On platforms that use OpenSSL < 1.1.1, register a locking callback so we can provide mutexes and stop it crashing horribly when the same RSA key is used concurrently on different threads

Changed default key size in the CryptoKeys plugin to 4096 bits.

diff --git a/Engine/Plugins/Editor/CryptoKeys/Source/CryptoKeys/Private/CryptoKeysHelpers.h b/Engine/Plugins/Editor/CryptoKeys/Source/CryptoKeys/Private/CryptoKeysHelpers.h
--- a/Engine/Plugins/Editor/CryptoKeys/Source/CryptoKeys/Private/CryptoKeysHelpers.h
+++ b/Engine/Plugins/Editor/CryptoKeys/Source/CryptoKeys/Private/CryptoKeysHelpers.h
@@ -23,5 +23,5 @@ namespace CryptoKeysHelpers
 	@param InNumKeyBits 		- How many bits to use for the RSA key
 	@returns true if key generation succeeded, false otherwise
 	*/
-	bool GenerateSigningKey(FString& OutPublicExponent, FString& OutPrivateExponent, FString& OutModulus, int32 NumKeyBits = 2048);
+	bool GenerateSigningKey(FString& OutPublicExponent, FString& OutPrivateExponent, FString& OutModulus, int32 NumKeyBits = 4096);
 }
\ No newline at end of file
diff --git a/Engine/Plugins/Editor/CryptoKeys/Source/CryptoKeysOpenSSL/Private/CryptoKeysOpenSSL.cpp b/Engine/Plugins/Editor/CryptoKeys/Source/CryptoKeysOpenSSL/Private/CryptoKeysOpenSSL.cpp
--- a/Engine/Plugins/Editor/CryptoKeys/Source/CryptoKeysOpenSSL/Private/CryptoKeysOpenSSL.cpp
+++ b/Engine/Plugins/Editor/CryptoKeys/Source/CryptoKeysOpenSSL/Private/CryptoKeysOpenSSL.cpp
@@ -17,6 +17,12 @@
 
 DEFINE_LOG_CATEGORY_STATIC(LogCryptoKeys, Log, All);
 
+#if !defined(OPENSSL_VERSION_NUMBER) || OPENSSL_VERSION_NUMBER < 0x10100000L
+#define USE_LEGACY_OPENSSL 1
+#else
+#define USE_LEGACY_OPENSSL 0
+#endif
+
 namespace CryptoKeysOpenSSL
 {
 	bool GenerateNewEncryptionKey(TArray<uint8>& OutKey)
@@ -35,9 +41,10 @@ namespace CryptoKeysOpenSSL
 	{
 		int32 NumBytes = BN_num_bytes(InNum);
 		check(NumBytes <= InKeySize);
-		OutBytes.SetNumZeroed(InKeySize);
-		BN_bn2bin(InNum, OutBytes.GetData() + (InKeySize - NumBytes));
-		Algo::Reverse(OutBytes); // bn2bin produces big endian data
+		OutBytes.SetNumZeroed(NumBytes);
+
+		BN_bn2bin(InNum, OutBytes.GetData());
+		Algo::Reverse(OutBytes);
 	}
 
 	bool GenerateNewSigningKey(TArray<uint8>& OutPublicExponent, TArray<uint8>& OutPrivateExponent, TArray<uint8>& OutModulus, int32 InNumKeyBits)
@@ -50,7 +57,7 @@ namespace CryptoKeysOpenSSL
 		BN_set_word(E, RSA_F4);
 		RSA_generate_key_ex(RSAKey, KeySize, E, nullptr);
 
-#if !defined(OPENSSL_VERSION_NUMBER) || OPENSSL_VERSION_NUMBER < 0x10100000L
+#if USE_LEGACY_OPENSSL
 		const BIGNUM* PublicModulus = RSAKey->n;
 		const BIGNUM* PublicExponent = RSAKey->e;
 		const BIGNUM* PrivateExponent = RSAKey->d;
diff --git a/Engine/Source/Developer/PakFileUtilities/PakFileUtilities.Build.cs b/Engine/Source/Developer/PakFileUtilities/PakFileUtilities.Build.cs
--- a/Engine/Source/Developer/PakFileUtilities/PakFileUtilities.Build.cs
+++ b/Engine/Source/Developer/PakFileUtilities/PakFileUtilities.Build.cs
@@ -6,7 +6,7 @@ public class PakFileUtilities : ModuleRules
 {
 	public PakFileUtilities(ReadOnlyTargetRules Target) : base(Target)
 	{
-		PrivateDependencyModuleNames.AddRange(new string[] { "Core", "PakFile", "Json", "Projects" });
+		PrivateDependencyModuleNames.AddRange(new string[] { "Core", "PakFile", "Json", "Projects", "RSA" });
 
         PrivateIncludePathModuleNames.AddRange(
             new string[] {
diff --git a/Engine/Source/Developer/PakFileUtilities/Private/PakFileUtilities.cpp b/Engine/Source/Developer/PakFileUtilities/Private/PakFileUtilities.cpp
--- a/Engine/Source/Developer/PakFileUtilities/Private/PakFileUtilities.cpp
+++ b/Engine/Source/Developer/PakFileUtilities/Private/PakFileUtilities.cpp
@@ -23,62 +23,6 @@
 
 IMPLEMENT_MODULE(FDefaultModuleImpl, PakFileUtilities);
 
-/**
- * Encryption keys: public and private
- */
-struct FKeyPair
-{
-	/** Public decryption key */
-	TPakRSAKey PublicKey;
-	/** Private encryption key */
-	TPakRSAKey PrivateKey;
-
-	bool IsValid() const
-	{
-		return !PrivateKey.Exponent.IsZero() &&
-		!PrivateKey.Modulus.IsZero() &&
-		!PublicKey.Exponent.IsZero() &&
-		!PublicKey.Modulus.IsZero();
-	}
-};
-
-bool TestKeys(FKeyPair& Pair)
-{
-	UE_LOG(LogPakFile, Display, TEXT("Testing signature keys."));
-
-	// Just some random values
-	static TPakRSAKey::TIntType TestData[] =
-	{
-		11,
-		253,
-		128,
-		234,
-		56,
-		89,
-		34,
-		179,
-		29,
-		1024,
-		(int64)(MAX_int32),
-		(int64)(MAX_uint32)-1
-	};
-
-	for (int32 TestIndex = 0; TestIndex < ARRAY_COUNT(TestData); ++TestIndex)
-	{
-		TPakRSAKey::TIntType EncryptedData = FEncryption::ModularPow(TestData[TestIndex], Pair.PrivateKey.Exponent, Pair.PrivateKey.Modulus);
-		TPakRSAKey::TIntType DecryptedData = FEncryption::ModularPow(EncryptedData, Pair.PublicKey.Exponent, Pair.PublicKey.Modulus);
-		if (TestData[TestIndex] != DecryptedData)
-		{
-			UE_LOG(LogPakFile, Error, TEXT("Keys do not properly encrypt/decrypt data (failed test with %lld)"), TestData[TestIndex].ToInt());
-			return false;
-		}
-	}
-
-	UE_LOG(LogPakFile, Display, TEXT("Signature keys check completed successfuly."));
-
-	return true;
-}
-
 struct FNamedAESKey
 {
 	FString Name;
@@ -93,7 +37,7 @@ struct FNamedAESKey
 
 struct FKeyChain
 {
-	FKeyPair SigningKey;
+	FRSA::TKeyPtr SigningKey;
 	TMap<FGuid, FNamedAESKey> EncryptionKeys;
 	const FNamedAESKey* MasterEncryptionKey = nullptr;
 };
@@ -1156,23 +1100,30 @@ TEncryptionInt ParseEncryptionIntFromJson(TSharedPtr<FJsonObject> InObj, const T
 	}
 }
 
-void ParseRSAKeyFromJson(TSharedPtr<FJsonObject> InObj, TPakRSAKey& OutKey)
+FRSA::TKeyPtr ParseRSAKeyFromJson(TSharedPtr<FJsonObject> InObj)
 {
-	FString ExponentBase64, ModulusBase64;
-	TArray<uint8> Exponent, Modulus;
-	if (InObj->TryGetStringField("Exponent", ExponentBase64)
-		&& InObj->TryGetStringField("Modulus", ModulusBase64))
+	TSharedPtr<FJsonObject> PublicKey = InObj->GetObjectField(TEXT("PublicKey"));
+	TSharedPtr<FJsonObject> PrivateKey = InObj->GetObjectField(TEXT("PrivateKey"));
+
+	FString PublicExponentBase64, PrivateExponentBase64, PublicModulusBase64, PrivateModulusBase64;
+
+	if (   PublicKey->TryGetStringField("Exponent", PublicExponentBase64)
+		&& PublicKey->TryGetStringField("Modulus", PublicModulusBase64)
+		&& PrivateKey->TryGetStringField("Exponent", PrivateExponentBase64)
+		&& PrivateKey->TryGetStringField("Modulus", PrivateModulusBase64))
 	{
-		FBase64::Decode(ExponentBase64, Exponent);
-		FBase64::Decode(ModulusBase64, Modulus);
+		check(PublicModulusBase64 == PrivateModulusBase64);
 
-		check(Exponent.Num() > 0 && Exponent.Num() <= TPakRSAKey::KeySizeInBytes);
-		check(Modulus.Num() > 0 && Modulus.Num() <= TPakRSAKey::KeySizeInBytes);
-		OutKey = TPakRSAKey(Exponent, Modulus);
+		TArray<uint8> PublicExponent, PrivateExponent, Modulus;
+		FBase64::Decode(PublicExponentBase64, PublicExponent);
+		FBase64::Decode(PrivateExponentBase64, PrivateExponent);
+		FBase64::Decode(PublicModulusBase64, Modulus);
+
+		return FRSA::CreateKey(PublicExponent, PrivateExponent, Modulus);
 	}
 	else
 	{
-		OutKey = TPakRSAKey();
+		return nullptr;
 	}
 }
 
@@ -1207,11 +1158,7 @@ void LoadKeyChainFromFile(const FString& InFilename, FKeyChain& OutCryptoSetting
 		const TSharedPtr<FJsonObject>* SigningKey = nullptr;
 		if (RootObject->TryGetObjectField(TEXT("SigningKey"), SigningKey))
 		{
-			TSharedPtr<FJsonObject> PublicKey = (*SigningKey)->GetObjectField(TEXT("PublicKey"));
-			TSharedPtr<FJsonObject> PrivateKey = (*SigningKey)->GetObjectField(TEXT("PrivateKey"));
-			ParseRSAKeyFromJson(PublicKey, OutCryptoSettings.SigningKey.PublicKey);
-			ParseRSAKeyFromJson(PrivateKey, OutCryptoSettings.SigningKey.PrivateKey);
-			check(OutCryptoSettings.SigningKey.PublicKey.Modulus == OutCryptoSettings.SigningKey.PrivateKey.Modulus);
+			OutCryptoSettings.SigningKey = ParseRSAKeyFromJson(*SigningKey);
 		}
 
 		const TArray<TSharedPtr<FJsonValue>>* SecondaryEncryptionKeyArray = nullptr;
@@ -1240,10 +1187,7 @@ void LoadKeyChainFromFile(const FString& InFilename, FKeyChain& OutCryptoSetting
 
 void LoadKeyChain(const TCHAR* CmdLine, FKeyChain& OutCryptoSettings)
 {
-	OutCryptoSettings.SigningKey.PrivateKey.Exponent.Zero();
-	OutCryptoSettings.SigningKey.PrivateKey.Modulus.Zero();
-	OutCryptoSettings.SigningKey.PublicKey.Exponent.Zero();
-	OutCryptoSettings.SigningKey.PublicKey.Modulus.Zero();
+	OutCryptoSettings.SigningKey.Reset();
 	OutCryptoSettings.EncryptionKeys.Empty();
 
 	// First, try and parse the keys from a supplied crypto key cache file
@@ -1261,6 +1205,8 @@ void LoadKeyChain(const TCHAR* CmdLine, FKeyChain& OutCryptoSettings)
 			&& FParse::Value(CmdLine, TEXT("enginedir="), EngineDir, false)
 			&& FParse::Value(CmdLine, TEXT("platform="), Platform, false))
 		{
+			UE_LOG(LogPakFile, Warning, TEXT("A legacy command line syntax is being used for crypto config. Please update to using the -cryptokey parameter as soon as possible as this mode is deprecated"));
+
 			FConfigFile EngineConfig;
 
 			FConfigCacheIni::LoadExternalIniFile(EngineConfig, TEXT("Engine"), *FPaths::Combine(EngineDir, TEXT("Config\\")), *FPaths::Combine(ProjectDir, TEXT("Config/")), true, *Platform);
@@ -1304,8 +1250,7 @@ void LoadKeyChain(const TCHAR* CmdLine, FKeyChain& OutCryptoSettings)
 					FBase64::Decode(PrivateExpBase64, PrivateExp);
 					FBase64::Decode(ModulusBase64, Modulus);
 
-					OutCryptoSettings.SigningKey.PrivateKey = TPakRSAKey(PrivateExp, Modulus);
-					OutCryptoSettings.SigningKey.PublicKey = TPakRSAKey(PublicExp, Modulus);
+					OutCryptoSettings.SigningKey = FRSA::CreateKey(PublicExp, PrivateExp, Modulus);
 
 					UE_LOG(LogPakFile, Display, TEXT("Parsed signature keys from config files."));
 				}
@@ -1346,10 +1291,11 @@ void LoadKeyChain(const TCHAR* CmdLine, FKeyChain& OutCryptoSettings)
 					ConfigFile.GetString(SectionName, TEXT("rsa.privateexp"), RSAPrivateExp);
 					ConfigFile.GetString(SectionName, TEXT("rsa.modulus"), RSAModulus);
 
-					OutCryptoSettings.SigningKey.PrivateKey.Exponent.Parse(RSAPrivateExp);
-					OutCryptoSettings.SigningKey.PrivateKey.Modulus.Parse(RSAModulus);
-					OutCryptoSettings.SigningKey.PublicKey.Exponent.Parse(RSAPublicExp);
-					OutCryptoSettings.SigningKey.PublicKey.Modulus = OutCryptoSettings.SigningKey.PrivateKey.Modulus;
+					//TODO: Fix me!
+					//OutSigningKey.PrivateKey.Exponent.Parse(RSAPrivateExp);
+					//OutSigningKey.PrivateKey.Modulus.Parse(RSAModulus);
+					//OutSigningKey.PublicKey.Exponent.Parse(RSAPublicExp);
+					//OutSigningKey.PublicKey.Modulus = OutSigningKey.PrivateKey.Modulus;
 
 					UE_LOG(LogPakFile, Display, TEXT("Parsed signature keys from config files."));
 				}
@@ -1383,6 +1329,8 @@ void LoadKeyChain(const TCHAR* CmdLine, FKeyChain& OutCryptoSettings)
 
 		if (EncryptionKeyString.Len() > 0)
 		{
+			UE_LOG(LogPakFile, Warning, TEXT("A legacy command line syntax is being used for crypto config. Please update to using the -cryptokey parameter as soon as possible as this mode is deprecated"));
+
 			FNamedAESKey NewKey;
 			NewKey.Name = TEXT("Default");
 			NewKey.Guid = FGuid();
@@ -1451,7 +1399,7 @@ FArchive* CreatePakWriter(const TCHAR* Filename, const FKeyChain& InKeyChain, bo
 		if (bSign)
 		{
 			UE_LOG(LogPakFile, Display, TEXT("Creating signed pak %s."), Filename);
-			Writer = new FSignedArchiveWriter(*Writer, Filename, InKeyChain.SigningKey.PublicKey, InKeyChain.SigningKey.PrivateKey);
+			Writer = new FSignedArchiveWriter(*Writer, Filename, InKeyChain.SigningKey);
 		}
 		else
 		{
diff --git a/Engine/Source/Developer/PakFileUtilities/Private/SignedArchiveWriter.cpp b/Engine/Source/Developer/PakFileUtilities/Private/SignedArchiveWriter.cpp
--- a/Engine/Source/Developer/PakFileUtilities/Private/SignedArchiveWriter.cpp
+++ b/Engine/Source/Developer/PakFileUtilities/Private/SignedArchiveWriter.cpp
@@ -5,14 +5,13 @@
 #include "Misc/SecureHash.h"
 #include "HAL/FileManager.h"
 
-FSignedArchiveWriter::FSignedArchiveWriter(FArchive& InPak, const FString& InPakFilename, const TPakRSAKey& InPublicKey, const TPakRSAKey& InPrivateKey)
+FSignedArchiveWriter::FSignedArchiveWriter(FArchive& InPak, const FString& InPakFilename, FRSA::TKeyPtr InSigningKey)
 : BufferArchive(Buffer)
 	, PakWriter(InPak)
 	, PakSignaturesFilename(FPaths::ChangeExtension(InPakFilename, TEXT("sig")))
 	, SizeOnDisk(0)
 	, PakSize(0)
-	, PublicKey(InPublicKey)
-	, PrivateKey(InPrivateKey)
+	, SigningKey(InSigningKey)
 {
 	Buffer.Reserve(FPakInfo::MaxChunkDataSize);
 }
@@ -46,19 +45,10 @@ bool FSignedArchiveWriter::Close()
 
 	FArchive* SignatureWriter = IFileManager::Get().CreateFileWriter(*PakSignaturesFilename);
 	FPakSignatureFile SignatureFile;
-	SignatureFile.Initialize(ChunkHashes, PrivateKey);
+	SignatureFile.SetChunkHashesAndSign(ChunkHashes, SigningKey);
 	SignatureFile.Serialize(*SignatureWriter);
 	delete SignatureWriter;
 
-	// Test reading back in
-	{
-		FArchive* SignatureReader = IFileManager::Get().CreateFileReader(*PakSignaturesFilename);
-		FPakSignatureFile SigFileTest;
-		SigFileTest.Serialize(*SignatureReader);
-		UE_CLOG(!SigFileTest.DecryptSignatureAndValidate(PublicKey, FString()), LogPakFile, Fatal, TEXT("Failed to validate RSA encrypted pak hash"));
-		delete SignatureReader;
-	}
-
 	return FArchive::Close();
 }
 
diff --git a/Engine/Source/Developer/PakFileUtilities/Private/SignedArchiveWriter.h b/Engine/Source/Developer/PakFileUtilities/Private/SignedArchiveWriter.h
--- a/Engine/Source/Developer/PakFileUtilities/Private/SignedArchiveWriter.h
+++ b/Engine/Source/Developer/PakFileUtilities/Private/SignedArchiveWriter.h
@@ -24,10 +24,8 @@ class FSignedArchiveWriter : public FArchive
 	int64 SizeOnDisk;
 	/** Data size (excluding signatures) */
 	int64 PakSize;
-	/** Decryption key */
-	TPakRSAKey PublicKey;
-	/** Encryption key */
-	TPakRSAKey PrivateKey;
+	/** Signing key */
+	FRSA::TKeyPtr SigningKey;
 	/** Hashes */
 	TArray<TPakChunkHash> ChunkHashes;
 
@@ -38,7 +36,7 @@ class FSignedArchiveWriter : public FArchive
 
 public:
 
-	FSignedArchiveWriter(FArchive& InPak, const FString& InPakFilename, const TPakRSAKey& InPublicKey, const TPakRSAKey& InPrivateKey);
+	FSignedArchiveWriter(FArchive& InPak, const FString& InPakFilename, FRSA::TKeyPtr InSigningKey);
 	virtual ~FSignedArchiveWriter();
 
 	// FArchive interface
diff --git a/Engine/Source/Programs/UnrealPak/UnrealPak.Build.cs b/Engine/Source/Programs/UnrealPak/UnrealPak.Build.cs
--- a/Engine/Source/Programs/UnrealPak/UnrealPak.Build.cs
+++ b/Engine/Source/Programs/UnrealPak/UnrealPak.Build.cs
@@ -8,7 +8,7 @@ public class UnrealPak : ModuleRules
 	{
 		PublicIncludePaths.Add("Runtime/Launch/Public");
 
-		PrivateDependencyModuleNames.AddRange(new string[] { "Core", "PakFile", "Json", "Projects", "PakFileUtilities" });
+		PrivateDependencyModuleNames.AddRange(new string[] { "Core", "PakFile", "Json", "Projects", "PakFileUtilities", "RSA" });
 
 		PrivateIncludePaths.Add("Runtime/Launch/Private");      // For LaunchEngineLoop.cpp include
 
diff --git a/Engine/Source/Runtime/Core/Public/Misc/RSA.h b/Engine/Source/Runtime/Core/Public/Misc/RSA.h
deleted file mode 100644
diff --git a/Engine/Source/Runtime/PakFile/PakFile.Build.cs b/Engine/Source/Runtime/PakFile/PakFile.Build.cs
--- a/Engine/Source/Runtime/PakFile/PakFile.Build.cs
+++ b/Engine/Source/Runtime/PakFile/PakFile.Build.cs
@@ -7,5 +7,6 @@ public class PakFile : ModuleRules
 	public PakFile(ReadOnlyTargetRules Target) : base(Target)
 	{
 		PrivateDependencyModuleNames.Add("Core");
+		PublicDependencyModuleNames.Add("RSA");
 	}
 }
diff --git a/Engine/Source/Runtime/PakFile/Private/IPlatformFilePak.cpp b/Engine/Source/Runtime/PakFile/Private/IPlatformFilePak.cpp
--- a/Engine/Source/Runtime/PakFile/Private/IPlatformFilePak.cpp
+++ b/Engine/Source/Runtime/PakFile/Private/IPlatformFilePak.cpp
@@ -213,16 +213,26 @@ void FPakPlatformFile::GetPakEncryptionKey(FAES::FAESKey& OutKey, const FGuid& I
 	}
 }
 
-void FPakPlatformFile::GetPakSigningKeys(TPakRSAKey& OutKey)
+TSharedPtr<FRSA::FKey, ESPMode::ThreadSafe> FPakPlatformFile::GetPakSigningKey()
 {
-	FCoreDelegates::FPakSigningKeysDelegate& Delegate = FCoreDelegates::GetPakSigningKeysDelegate();
-	if (Delegate.IsBound())
+	static TSharedPtr<FRSA::FKey, ESPMode::ThreadSafe> Key;
+	static FCriticalSection Lock;
+	Lock.Lock();
+
+	if (!Key.IsValid())
 	{
-		TArray<uint8> Exponent;
-		TArray<uint8> Modulus;
-		Delegate.Execute(Exponent, Modulus);
-		OutKey = TPakRSAKey(Exponent, Modulus);
+		FCoreDelegates::FPakSigningKeysDelegate& Delegate = FCoreDelegates::GetPakSigningKeysDelegate();
+		if (Delegate.IsBound())
+		{
+			TArray<uint8> Exponent;
+			TArray<uint8> Modulus;
+			Delegate.Execute(Exponent, Modulus);
+			Key = FRSA::CreateKey(Exponent, TArray<uint8>(), Modulus);
+		}
 	}
+	
+	Lock.Unlock();
+	return Key;
 }
 
 DECLARE_DWORD_ACCUMULATOR_STAT(TEXT("PakCache Sync Decrypts (Uncompressed Path)"), STAT_PakCache_SyncDecrypts, STATGROUP_PakFile);
@@ -1126,17 +1136,16 @@ class FPakPrecacher
 	uint32 Loads;
 	uint32 Frees;
 	uint64 LoadSize;
-	TPakRSAKey PublicKey;
-	bool bSigned;
+	FRSA::TKeyPtr SigningKey;
 	EAsyncIOPriorityAndFlags AsyncMinPriority;
 	FCriticalSection SetAsyncMinimumPriorityScopeLock;
 public:
 
-	static void Init(IPlatformFile* InLowerLevel, const TPakRSAKey& InEncryptionKey)
+	static void Init(IPlatformFile* InLowerLevel, FRSA::TKeyPtr InSigningKey)
 	{
 		if (!PakPrecacherSingleton)
 		{
-			verify(!FPlatformAtomics::InterlockedCompareExchangePointer((void**)&PakPrecacherSingleton, new FPakPrecacher(InLowerLevel, InEncryptionKey), nullptr));
+			verify(!FPlatformAtomics::InterlockedCompareExchangePointer((void**)&PakPrecacherSingleton, new FPakPrecacher(InLowerLevel, InSigningKey), nullptr));
 		}
 		check(PakPrecacherSingleton);
 	}
@@ -1172,7 +1181,7 @@ public:
 		return *PakPrecacherSingleton;
 	}
 
-	FPakPrecacher(IPlatformFile* InLowerLevel, const TPakRSAKey& InPublicKey)
+	FPakPrecacher(IPlatformFile* InLowerLevel, FRSA::TKeyPtr InSigningKey)
 		: LowerLevel(InLowerLevel)
 		, LastReadRequest(0)
 		, NextUniqueID(1)
@@ -1182,8 +1191,7 @@ public:
 		, Loads(0)
 		, Frees(0)
 		, LoadSize(0)
-		, PublicKey(InPublicKey)
-		, bSigned(!InPublicKey.Exponent.IsZero() && !InPublicKey.Modulus.IsZero())
+		, SigningKey(InSigningKey)
 		, AsyncMinPriority(AIOP_MIN)
 	{
 		check(LowerLevel && FPlatformProcess::SupportsMultithreading());
@@ -1223,7 +1231,7 @@ public:
 
 			FPakData& Pak = CachedPakData[*PakIndexPtr];
 
-			if (bSigned)
+			if (SigningKey.IsValid())
 			{
 				// Load signature data
 				FString SignaturesFilename = FPaths::ChangeExtension(*PakFilename, TEXT("sig"));
@@ -1232,7 +1240,7 @@ public:
 				FArchiveFileReaderGeneric* Reader = new FArchiveFileReaderGeneric(SignaturesFile, *SignaturesFilename, SignaturesFile->Size());
 				Pak.Signatures.Serialize(*Reader);
 				delete Reader;
-				Pak.Signatures.DecryptSignatureAndValidate(PublicKey, PakFilename);
+				Pak.Signatures.DecryptSignatureAndValidate(SigningKey, PakFilename);
 
 				// Check that we have the correct match between signature and pre-cache granularity
 				int64 NumPakChunks = Align(PakFileSize, FPakInfo::MaxChunkDataSize) / FPakInfo::MaxChunkDataSize;
@@ -2056,7 +2064,7 @@ private: // below here we assume CachedFilesScopeLock until we get to the next s
 		FAsyncFileCallBack CallbackFromLower =
 			[this, IndexToFill, bDoCheck](bool bWasCanceled, IAsyncReadRequest* Request)
 		{
-			if (bSigned && bDoCheck)
+			if (SigningKey.IsValid() && bDoCheck)
 			{
 				StartSignatureCheck(bWasCanceled, Request, IndexToFill);
 			}
@@ -5381,11 +5389,8 @@ bool FPakPlatformFile::Initialize(IPlatformFile* Inner, const TCHAR* CmdLine)
 	GameUserSettingsIniFilename = TEXT("GameUserSettings.ini");
 #endif
 
-	TPakRSAKey DecryptionKey;
-	GetPakSigningKeys(DecryptionKey);
-
 	// signed if we have keys, and are not running with fileopenlog (currently results in a deadlock).
-	bSigned = !DecryptionKey.Exponent.IsZero() && !DecryptionKey.Modulus.IsZero() && !FParse::Param(FCommandLine::Get(), TEXT("fileopenlog"));;
+	bSigned = GetPakSigningKey().IsValid() && !FParse::Param(FCommandLine::Get(), TEXT("fileopenlog"));;
 		
 	// Find and mount pak files from the specified directories.
 	TArray<FString> PakFolders;
@@ -5442,10 +5447,7 @@ void FPakPlatformFile::InitializeNewAsyncIO()
 #if !WITH_EDITOR
 	if (FPlatformProcess::SupportsMultithreading() && !FParse::Param(FCommandLine::Get(), TEXT("FileOpenLog")))
 	{
-		TPakRSAKey DecryptionKey;
-		GetPakSigningKeys(DecryptionKey);
-
-		FPakPrecacher::Init(LowerLevel, DecryptionKey);
+		FPakPrecacher::Init(LowerLevel, GetPakSigningKey());
 	}
 	else
 #endif
@@ -5936,6 +5938,7 @@ public:
 	virtual void StartupModule() override
 	{
 		Singleton = MakeUnique<FPakPlatformFile>();
+		FModuleManager::LoadModuleChecked<IModuleInterface>(TEXT("RSA"));
 	}
 
 	virtual void ShutdownModule() override
diff --git a/Engine/Source/Runtime/PakFile/Public/IPlatformFilePak.h b/Engine/Source/Runtime/PakFile/Public/IPlatformFilePak.h
--- a/Engine/Source/Runtime/PakFile/Public/IPlatformFilePak.h
+++ b/Engine/Source/Runtime/PakFile/Public/IPlatformFilePak.h
@@ -11,7 +11,7 @@
 #include "Templates/UniquePtr.h"
 #include "Math/BigInt.h"
 #include "Misc/AES.h"
-#include "Misc/RSA.h"
+#include "RSA.h"
 #include "Misc/SecureHash.h"
 #include "GenericPlatform/GenericPlatformChunkInstall.h"
 
@@ -33,9 +33,6 @@ typedef uint32 TPakChunkHash;
 typedef FSHAHash TPakChunkHash;
 #endif
 
-// Define the key size that we use for pak signing purposes. Must match with the key size created in the cryptokeys editor plugin
-typedef FRSAKey<2048> TPakRSAKey;
-
 PAKFILE_API TPakChunkHash ComputePakChunkHash(const void* InData, int64 InDataSizeInBytes);
 
 struct FPakChunkSignatureCheckFailedData
@@ -1517,9 +1514,9 @@ public:
 	static void GetPakEncryptionKey(FAES::FAESKey& OutKey, const FGuid& InEncryptionKeyGuid);
 
 	/**
-	* Helper function for accessing pak signing keys
+	* Helper function for accessing pak the signing public key
 	*/
-	static void GetPakSigningKeys(TPakRSAKey& OutKey);
+	static TSharedPtr<FRSA::FKey, ESPMode::ThreadSafe> GetPakSigningKey();
 
 	/**
 	 * Constructor.
@@ -2364,11 +2361,11 @@ struct FPakSignatureFile
 	/**
 	 * Initialize and hash the CRC list then use the provided private key to encrypt the hash
 	 */
-	void Initialize(const TArray<TPakChunkHash>& InChunkHashes, const TPakRSAKey& InPrivateKey)
+	void SetChunkHashesAndSign(const TArray<TPakChunkHash>& InChunkHashes, const FRSA::TKeyPtr InKey)
 	{
 		ChunkHashes = InChunkHashes;
 		DecryptedHash = ComputeCurrentMasterHash();
-		TPakRSAKey::Encrypt(DecryptedHash.Hash, ARRAY_COUNT(FSHAHash::Hash), EncryptedHash, InPrivateKey);
+		FRSA::Encrypt(FRSA::EKeyType::Private, DecryptedHash.Hash, ARRAY_COUNT(FSHAHash::Hash), EncryptedHash, InKey);
 	}
 
 	/**
@@ -2403,19 +2400,19 @@ struct FPakSignatureFile
 	 */
 	bool DecryptSignatureAndValidate(const FString& InFilename)
 	{
-		TPakRSAKey PublicKey;
-		FPakPlatformFile::GetPakSigningKeys(PublicKey);
+		TSharedPtr<FRSA::FKey, ESPMode::ThreadSafe> PublicKey = FPakPlatformFile::GetPakSigningKey();
 		return DecryptSignatureAndValidate(PublicKey, InFilename);
 	}
 
 	/**
 	 * Decrypt the chunk CRCs hash and validate that it matches the current one
 	 */
-	bool DecryptSignatureAndValidate(TPakRSAKey& InPublicKey, const FString& InFilename)
+	bool DecryptSignatureAndValidate(FRSA::TKeyPtr InKey, const FString& InFilename)
 	{
-		TPakRSAKey::Decrypt(EncryptedHash, DecryptedHash.Hash, ARRAY_COUNT(FSHAHash::Hash), InPublicKey);
+		FRSA::Decrypt(FRSA::EKeyType::Public, EncryptedHash, DecryptedHash.Hash, ARRAY_COUNT(FSHAHash::Hash), InKey);
 
-		if (DecryptedHash != ComputeCurrentMasterHash())
+		FSHAHash CurrentHash = ComputeCurrentMasterHash();
+		if (DecryptedHash != CurrentHash)
 		{
 			FPakPlatformFile::GetPakMasterSignatureTableCheckFailureHandler().Broadcast(InFilename);
 			return false;
diff --git a/Engine/Source/Runtime/RSA/Private/RSA.cpp b/Engine/Source/Runtime/RSA/Private/RSA.cpp
new file mode 100644
--- /dev/null
+++ b/Engine/Source/Runtime/RSA/Private/RSA.cpp
@@ -0,0 +1,426 @@
+// Copyright 1998-2019 Epic Games, Inc. All Rights Reserved.
+
+#include "RSA.h"
+#include "Misc/ScopeLock.h"
+#include "Modules/ModuleManager.h"
+#include "Math/BigInt.h"
+
+#ifndef RSA_USE_OPENSSL
+#define RSA_USE_OPENSSL 0
+#endif
+
+#if RSA_USE_OPENSSL
+#if PLATFORM_WINDOWS
+#include "Windows/AllowWindowsPlatformTypes.h"
+#endif
+
+#include <openssl/rsa.h>
+#include <openssl/bn.h>
+#include <openssl/err.h>
+
+#if PLATFORM_WINDOWS
+#include "Windows/HideWindowsPlatformTypes.h"
+#endif
+
+// Some platforms were upgraded to OpenSSL 1.1.1 while the others were left on a previous version. There are some minor differences we have to account for
+// in the older version, so declare a handy define that we can use to gate the code
+#if !defined(OPENSSL_VERSION_NUMBER) || OPENSSL_VERSION_NUMBER < 0x10100000L
+#define USE_LEGACY_OPENSSL 1
+#else
+#define USE_LEGACY_OPENSSL 0
+#endif
+
+struct FRSA::FKey
+{
+	FKey() 
+	{
+	}
+
+	~FKey()
+	{
+		RSA_free(RSAKey);
+	}
+
+	int32 GetKeySizeInBits() const
+	{
+		return KeySize;
+	}
+
+	int32 GetKeySizeInBytes() const
+	{
+		return KeySizeInBytes;
+	}
+
+	int32 GetMaxDataSize() const
+	{
+		return MaxDataSize;
+	}
+
+	enum class EType
+	{
+		Public,
+		Private,
+	};
+
+	int32 KeySize;
+	int32 KeySizeInBytes;
+	int32 MaxDataSize;
+	RSA* RSAKey;
+	EType Type;
+
+#if USE_LEGACY_OPENSSL
+	FCriticalSection Lock;
+#endif
+};
+
+/**
+ * Wrapper function for extracting the bytes of an OpenSSL big num in a consistent form
+ */
+void LoadBinaryIntoBigNum(const uint8* InData, int64 InDataSize, BIGNUM* InBigNum)
+{
+#if USE_LEGACY_OPENSSL
+	TArray<uint8> Bytes(InData, InDataSize);
+	Algo::Reverse(Bytes);
+	BN_bin2bn(Bytes.GetData(), Bytes.Num(), InBigNum);
+#else
+	BN_lebin2bn(InData, InDataSize, InBigNum);
+#endif
+}
+
+FRSA::TKeyPtr FRSA::CreateKey(const TArray<uint8>& InPublicExponent, const TArray<uint8>& InPrivateExponent, const TArray<uint8>& InModulus)
+{
+	FRSA::TKeyPtr Key = MakeShared<FRSA::FKey, FRSA::KeyThreadMode>();
+	Key->RSAKey = RSA_new();
+	BIGNUM* PublicExponent = InPublicExponent.Num() > 0 ? BN_new() : nullptr;
+	BIGNUM* PrivateExponent = InPrivateExponent.Num() > 0 ? BN_new() : nullptr;
+	BIGNUM* Modulus = BN_new();
+	
+	if (InPublicExponent.Num())
+	{
+		LoadBinaryIntoBigNum(InPublicExponent.GetData(), InPublicExponent.Num(), PublicExponent);
+		Key->Type = FRSA::FKey::EType::Public;
+	}
+	
+	if (InPrivateExponent.Num())
+	{
+		LoadBinaryIntoBigNum(InPrivateExponent.GetData(), InPrivateExponent.Num(), PrivateExponent);
+		Key->Type = FRSA::FKey::EType::Public;
+	}
+
+	LoadBinaryIntoBigNum(InModulus.GetData(), InModulus.Num(), Modulus);
+#if USE_LEGACY_OPENSSL
+	Key->RSAKey->n = Modulus;
+	Key->RSAKey->e = PublicExponent;
+	Key->RSAKey->d = PrivateExponent;
+#else
+	RSA_set0_key(Key->RSAKey, Modulus, PublicExponent, PrivateExponent);
+#endif
+
+	Key->KeySizeInBytes = RSA_size(Key->RSAKey);
+	Key->KeySize = Key->KeySizeInBytes * 8;
+	Key->MaxDataSize = Key->KeySizeInBytes - RSA_PKCS1_PADDING_SIZE;
+	return Key;
+}
+
+bool FRSA::Encrypt(EKeyType InKeyType, const uint8* InSource, int32 InSourceSizeInBytes, TArray<uint8>& OutDestination, FRSA::TKeyPtr InKey)
+{
+	OutDestination.SetNum(InKey->KeySizeInBytes);
+
+	int NumEncryptedBytes = 0;
+	
+	switch (InKeyType)
+	{
+	case FRSA::EKeyType::Public:
+	{
+		NumEncryptedBytes = RSA_public_encrypt(InSourceSizeInBytes, InSource, OutDestination.GetData(), InKey->RSAKey, RSA_PKCS1_PADDING);
+		break;
+	}
+
+	case FRSA::EKeyType::Private:
+	{
+		NumEncryptedBytes = RSA_private_encrypt(InSourceSizeInBytes, InSource, OutDestination.GetData(), InKey->RSAKey, RSA_PKCS1_PADDING);
+		break;
+	}
+	}
+	
+	if (NumEncryptedBytes == InKey->KeySizeInBytes)
+	{
+		return true;
+	}
+	else
+	{
+		OutDestination.Empty();
+		return false;
+	}
+}
+
+bool FRSA::Decrypt(EKeyType InKeyType, const TArray<uint8>& InSource, uint8* OutDestination, int32 OutDestinationSizeInBytes, FRSA::TKeyPtr InKey)
+{
+	if (InSource.Num() == InKey->GetKeySizeInBytes() && OutDestinationSizeInBytes <= InKey->GetKeySizeInBytes())
+	{
+		int NumDecryptedBytes = 0;
+
+		switch (InKey->Type)
+		{
+		case FKey::EType::Public:
+		{
+#if USE_LEGACY_OPENSSL
+			FScopeLock Lock(&InKey->Lock);
+#endif
+			NumDecryptedBytes = RSA_public_decrypt(InSource.Num(), InSource.GetData(), OutDestination, InKey->RSAKey, RSA_PKCS1_PADDING);
+			break;
+		}
+
+		case FKey::EType::Private:
+		{
+#if USE_LEGACY_OPENSSL
+			FScopeLock Lock(&InKey->Lock);
+#endif
+			NumDecryptedBytes = RSA_private_decrypt(InSource.Num(), InSource.GetData(), OutDestination, InKey->RSAKey, RSA_PKCS1_PADDING);
+			break;
+		}
+		}
+
+		if (NumDecryptedBytes >= 0 && NumDecryptedBytes <= OutDestinationSizeInBytes)
+		{
+			return true;
+		}
+		else
+		{
+			long Error = ERR_get_error();
+			const char* Message = ERR_error_string(Error, nullptr);
+			return false;
+		}
+	}
+
+	return false;
+}
+
+#else // !RSA_USE_OPENSSL
+
+//#pragma message ("Using TBigInt for RSA")
+#include "Math/RandomStream.h"
+
+struct FRSA::FKey
+{
+	virtual int32 GetKeySizeInBits() const = 0;
+	virtual int32 GetKeySizeInBytes() const = 0;
+	virtual int32 GetMaxDataSize() const = 0;
+	virtual bool Encrypt(FRSA::EKeyType InKeyType, const uint8* InSource, int32 InSourceSizeInBytes, TArray<uint8>& OutDestination) const = 0;
+	virtual bool Decrypt(FRSA::EKeyType InKeyType, const TArray<uint8>& InSource, uint8* OutDestination, int32 OutDestinationSizeInBytes) const = 0;
+};
+
+template <int32 KEY_SIZE>
+struct FFixedKey : public FRSA::FKey
+{
+	typedef TBigInt<KEY_SIZE * 2, false> TInternalBigInt;
+	TInternalBigInt PublicExponent;
+	TInternalBigInt PrivateExponent;
+	TInternalBigInt Modulus;
+
+	FFixedKey(const TArray<uint8>& InPublicExponent, const TArray<uint8>& InPrivateExponent, const TArray<uint8>& InModulus)
+	{
+		if (InPublicExponent.Num())
+		{
+			PublicExponent = TInternalBigInt(InPublicExponent.GetData(), InPublicExponent.Num());
+		}
+
+		if (InPrivateExponent.Num())
+		{
+			PrivateExponent = TInternalBigInt(InPrivateExponent.GetData(), InPrivateExponent.Num());
+		}
+
+		if (InModulus.Num())
+		{
+			Modulus = TInternalBigInt(InModulus.GetData(), InModulus.Num());
+		}
+	}
+
+	int32 GetKeySizeInBits() const override
+	{
+		return KEY_SIZE;
+	}
+
+	int32 GetKeySizeInBytes() const override
+	{
+		return KEY_SIZE / 8;
+	}
+
+	/** 
+		* Return the maximum amount of data that can be encrypted within the key, as you would if you were using a proper 
+		* RSA padding scheme. Because this is a legacy system and should be deprecated soon, we're just going to use random byte
+		* padding, but we'll reflect the correct RSA rules anyway
+		*/
+	int32 GetMaxDataSize() const override
+	{
+		return GetKeySizeInBytes() - 11;
+	}
+
+	bool Encrypt(FRSA::EKeyType InKeyType, const uint8* InSource, int32 InSourceSizeInBytes, TArray<uint8>& OutDestination) const override
+	{
+		if (InSourceSizeInBytes <= GetMaxDataSize())
+		{
+			TInternalBigInt Source(InSource, InSourceSizeInBytes);
+			int32 PaddingBytesRequired = GetKeySizeInBytes() - InSourceSizeInBytes;
+			uint8* PaddingStart = (uint8*)Source.GetBits() + InSourceSizeInBytes;
+			// Not perfect, but an improvement on before, and soon to be defuncted anyway
+			for (int32 PaddingIndex = 0; PaddingIndex < PaddingBytesRequired; ++PaddingIndex)
+			{
+				PaddingStart[PaddingIndex] = (uint8)FMath::RandRange(0, 255);
+			}
+			PaddingStart[PaddingBytesRequired - 1] &= 0x3f; // keep top two bits clear because this seems to cause the decryption to fail. TODO: Find out why... presumably some bug in TBigInt
+			TInternalBigInt Result = FEncryption::ModularPow(Source, InKeyType == FRSA::EKeyType::Public ? PublicExponent : PrivateExponent, Modulus);
+			TInternalBigInt Test = FEncryption::ModularPow(Result, InKeyType == FRSA::EKeyType::Public ? PrivateExponent : PublicExponent, Modulus);
+			check(Source == Test);
+			OutDestination.SetNum(GetKeySizeInBytes());
+			FMemory::Memcpy(OutDestination.GetData(), Result.GetBits(), GetKeySizeInBytes());
+		}
+		else
+		{
+			OutDestination.Empty();
+		}
+
+		return true;
+	}
+
+	bool Decrypt(FRSA::EKeyType InKeyType, const TArray<uint8>& InSource, uint8* OutDestination, int32 OutDestinationSizeInBytes) const override
+	{
+		if (InSource.Num() == GetKeySizeInBytes() && OutDestinationSizeInBytes <= GetKeySizeInBytes())
+		{
+			TInternalBigInt Source(InSource.GetData(), InSource.Num());
+			TInternalBigInt Result = FEncryption::ModularPow(Source, InKeyType == FRSA::EKeyType::Public ? PublicExponent : PrivateExponent, Modulus);
+			FMemory::Memcpy(OutDestination, Result.GetBits(), OutDestinationSizeInBytes);
+			return true;
+		}
+		else
+		{
+			FMemory::Memset(OutDestination, 0, OutDestinationSizeInBytes);
+			return false;
+		}
+	}
+};
+
+int32 NumElementsIgnoringTrailingZeroes(const TArray<uint8>& InBytes)
+{
+	int32 Count = InBytes.Num();
+
+	while (Count > 0 && InBytes[Count - 1] == 0)
+	{
+		Count--;
+	}
+
+	return Count;
+}
+
+FRSA::TKeyPtr FRSA::CreateKey(const TArray<uint8>& InPublicExponent, const TArray<uint8>& InPrivateExponent, const TArray<uint8>& InModulus)
+{
+	// The key data generated by openSSL is little endian, which is great for the latest version of OpenSSL and for our TBigInt. The older
+	// version of OpenSSL still in use by Linux/Max doesn't have the same little-endian import functions so we have to swizzle the data 
+	// in that case. OpenSSL should be upgraded on those platforms at some point, at which point we can remove this sub-optimal behavior
+	int32 NumPublicExponentBytes = NumElementsIgnoringTrailingZeroes(InPublicExponent);
+	int32 NumPrivateExponentBytes = NumElementsIgnoringTrailingZeroes(InPrivateExponent);
+	int32 NumModulusBytes = NumElementsIgnoringTrailingZeroes(InModulus);
+
+	int32 RequiredNumBytes = FMath::Max(NumModulusBytes, FMath::Max(NumPublicExponentBytes, NumPrivateExponentBytes));
+	RequiredNumBytes = FMath::RoundUpToPowerOfTwo(RequiredNumBytes);
+	int32 RequiredNumBits = RequiredNumBytes * 8;
+
+	// With the legacy fixed key system, which is based on a templated large integer class, we can't be that
+	// dynamic with the key lengths. Eventually, we'll move over to the totally dynamic version, but this is still an
+	// expansion on the previous functionality
+	FRSA::TKeyPtr Result;
+	
+	switch (RequiredNumBits)
+	{
+	case 4096: Result = MakeShared<FFixedKey<4096>, FRSA::KeyThreadMode>(InPublicExponent, InPrivateExponent, InModulus); break;
+	case 2048: Result = MakeShared<FFixedKey<2048>, FRSA::KeyThreadMode>(InPublicExponent, InPrivateExponent, InModulus); break;
+	case 1024: Result = MakeShared<FFixedKey<1024>, FRSA::KeyThreadMode>(InPublicExponent, InPrivateExponent, InModulus); break;
+	case 512: Result = MakeShared<FFixedKey<512>, FRSA::KeyThreadMode>(InPublicExponent, InPrivateExponent, InModulus); break;
+	case 256: Result = MakeShared<FFixedKey<256>, FRSA::KeyThreadMode>(InPublicExponent, InPrivateExponent, InModulus); break;
+	default: break;
+	}
+		
+	return Result;
+}
+
+bool FRSA::Encrypt(FRSA::EKeyType InKeyType, const uint8* InSource, int32 InSourceSizeInBytes, TArray<uint8>& OutDestination, FRSA::TKeyPtr InKey)
+{
+	return InKey->Encrypt(InKeyType, InSource, InSourceSizeInBytes, OutDestination);
+}
+
+bool FRSA::Decrypt(FRSA::EKeyType InKeyType, const TArray<uint8>& InSource, uint8* OutDestination, int32 OutDestinationSizeInBytes, FRSA::TKeyPtr InKey)
+{
+	return InKey->Decrypt(InKeyType, InSource, OutDestination, OutDestinationSizeInBytes);
+}
+#endif
+
+int32 FRSA::GetKeySizeInBits(FRSA::TKeyPtr InKey)
+{
+	return InKey->GetKeySizeInBits();
+}
+
+int32 FRSA::GetMaxDataSizeInBytes(FRSA::TKeyPtr InKey)
+{
+	return InKey->GetMaxDataSize();
+}
+
+bool FRSA::Encrypt(EKeyType InKeyType, const TArray<uint8>& InSource, TArray<uint8>& OutDestination, FRSA::TKeyPtr InKey)
+{
+	return Encrypt(InKeyType, InSource.GetData(), InSource.Num(), OutDestination, InKey);
+}
+
+bool FRSA::Decrypt(EKeyType InKeyType, const TArray<uint8>& InSource, TArray<uint8>& OutDestination, FRSA::TKeyPtr InKey)
+{
+	OutDestination.SetNum(InKey->GetKeySizeInBytes());
+	return Decrypt(InKeyType, InSource, OutDestination.GetData(), OutDestination.Num(), InKey);
+}
+
+#if RSA_USE_OPENSSL && USE_LEGACY_OPENSSL
+/**
+ * Module interface for RSA.
+ */
+class FRSAModule : public FDefaultModuleImpl
+{
+public:
+
+	virtual void StartupModule() override
+	{
+		// Provide a locking callback so that an RSA key can be used from multiple threads safely.
+		GetLocks().AddDefaulted(CRYPTO_num_locks());
+		CRYPTO_set_locking_callback(LockingCallback);
+	}
+
+	virtual void ShutdownModule() override
+	{
+		CRYPTO_set_locking_callback(nullptr);
+
+		// Just in case
+		for (FCriticalSection& CriticalSection : GetLocks())
+		{
+			CriticalSection.Unlock();
+		}
+	}
+
+	static TArray<FCriticalSection>& GetLocks()
+	{
+		static TArray<FCriticalSection> Locks;
+		return Locks;
+	}
+
+	static void LockingCallback(int mode, int n, const char *file, int line)
+	{
+		TArray<FCriticalSection>& Locks = GetLocks();
+		check(Locks.Num() > 0);
+
+		if (mode & CRYPTO_LOCK)
+			Locks[n].Lock();
+		else
+			Locks[n].Unlock();
+	}
+};
+
+IMPLEMENT_MODULE(FRSAModule, RSA);
+#else
+// OpenSSL 1.1.1+ handles thread safety of keys internally, so there is no need to provide any custom handling and we can just use a default module implementation
+IMPLEMENT_MODULE(FDefaultModuleImpl, RSA);
+#endif
\ No newline at end of file
diff --git a/Engine/Source/Runtime/RSA/Public/RSA.h b/Engine/Source/Runtime/RSA/Public/RSA.h
new file mode 100644
--- /dev/null
+++ b/Engine/Source/Runtime/RSA/Public/RSA.h
@@ -0,0 +1,49 @@
+// Copyright 1998-2019 Epic Games, Inc. All Rights Reserved.
+#pragma once
+
+#include "CoreTypes.h"
+#include "Templates/SharedPointer.h"
+
+struct RSA_API FRSA
+{
+	struct FKey;
+	static const ESPMode KeyThreadMode = ESPMode::ThreadSafe;
+	typedef TSharedPtr<FKey, KeyThreadMode> TKeyPtr;
+
+	/**
+	 * Enumerate which key to use when performing encrypt/decrpt operations
+	 */
+	enum class EKeyType
+	{
+		Public,
+		Private,
+	};
+
+	/**
+	 * Create a new RSA public/private key from the supplied exponents and modulus binary data. Each of these arrays should contain a single little endian
+	 * large integer value
+	 */
+	static TKeyPtr CreateKey(const TArray<uint8>& InPublicExponent, const TArray<uint8>& InPrivateExponent, const TArray<uint8>& InModulus);
+
+	/**
+	 * Returns the size in bits of the supplied key
+	 */
+	static int32 GetKeySizeInBits(TKeyPtr InKey);
+
+	/**
+	 * Returns the maximum number of bytes that can be encrypted in a single payload
+	 */
+	static int32 GetMaxDataSizeInBytes(TKeyPtr InKey);
+
+	/**
+	 * Encrypt the supplied byte data using the given key
+	 */
+	static bool Encrypt(EKeyType InKeyType, const uint8* InSource, int32 InSourceSizeInBytes, TArray<uint8>& OutDestination, TKeyPtr InKey);
+	static bool Encrypt(EKeyType InKeyType, const TArray<uint8>& InSource, TArray<uint8>& OutDestination, TKeyPtr InKey);
+
+	/**
+	 * Decrypt the supplied byte data using the given key
+	 */
+	static bool Decrypt(EKeyType InKeyType, const TArray<uint8>& InSource, uint8* OutDestination, int32 OutDestinationSizeInBytes, TKeyPtr InKey);
+	static bool Decrypt(EKeyType InKeyType, const TArray<uint8>& InSource, TArray<uint8>& OutDestination, TKeyPtr InKey);
+};
\ No newline at end of file
diff --git a/Engine/Source/Runtime/RSA/RSA.Build.cs b/Engine/Source/Runtime/RSA/RSA.Build.cs
new file mode 100644
--- /dev/null
+++ b/Engine/Source/Runtime/RSA/RSA.Build.cs
@@ -0,0 +1,29 @@
+// Copyright 1998-2019 Epic Games, Inc. All Rights Reserved.
+
+using UnrealBuildTool;
+
+public class RSA : ModuleRules
+{
+	public RSA(ReadOnlyTargetRules Target) : base(Target)
+	{
+		PrivateDependencyModuleNames.Add("Core");
+
+        switch (Target.Platform)
+        {
+			case UnrealTargetPlatform.Win64:
+			case UnrealTargetPlatform.Win32:
+			case UnrealTargetPlatform.Mac:
+			case UnrealTargetPlatform.Linux:
+				{
+					PrivateDependencyModuleNames.Add("OpenSSL");
+					PrivateDefinitions.Add("RSA_USE_OPENSSL=1");
+					break;
+				}
+			default:
+				{
+					PrivateDefinitions.Add("RSA_USE_OPENSSL=0");
+					break;
+				}
+		}
+	}
+}
