Minor bug fixes in RSA implementation

diff --git a/Engine/Source/Runtime/RSA/Private/RSA.cpp b/Engine/Source/Runtime/RSA/Private/RSA.cpp
--- a/Engine/Source/Runtime/RSA/Private/RSA.cpp
+++ b/Engine/Source/Runtime/RSA/Private/RSA.cpp
@@ -33,6 +33,10 @@
 struct FRSA::FKey
 {
 	FKey() 
+	: KeySize(0)
+	, KeySizeInBytes(0)
+	, MaxDataSize(0)
+	, RSAKey(RSA_new())
 	{
 	}
 
@@ -56,21 +60,10 @@ struct FRSA::FKey
 		return MaxDataSize;
 	}
 
-	enum class EType
-	{
-		Public,
-		Private,
-	};
-
 	int32 KeySize;
 	int32 KeySizeInBytes;
 	int32 MaxDataSize;
 	RSA* RSAKey;
-	EType Type;
-
-#if USE_LEGACY_OPENSSL
-	FCriticalSection Lock;
-#endif
 };
 
 /**
@@ -90,7 +83,6 @@ void LoadBinaryIntoBigNum(const uint8* InData, int64 InDataSize, BIGNUM* InBigNu
 FRSA::TKeyPtr FRSA::CreateKey(const TArray<uint8>& InPublicExponent, const TArray<uint8>& InPrivateExponent, const TArray<uint8>& InModulus)
 {
 	FRSA::TKeyPtr Key = MakeShared<FRSA::FKey, FRSA::KeyThreadMode>();
-	Key->RSAKey = RSA_new();
 	BIGNUM* PublicExponent = InPublicExponent.Num() > 0 ? BN_new() : nullptr;
 	BIGNUM* PrivateExponent = InPrivateExponent.Num() > 0 ? BN_new() : nullptr;
 	BIGNUM* Modulus = BN_new();
@@ -98,13 +90,11 @@ FRSA::TKeyPtr FRSA::CreateKey(const TArray<uint8>& InPublicExponent, const TArra
 	if (InPublicExponent.Num())
 	{
 		LoadBinaryIntoBigNum(InPublicExponent.GetData(), InPublicExponent.Num(), PublicExponent);
-		Key->Type = FRSA::FKey::EType::Public;
 	}
 	
 	if (InPrivateExponent.Num())
 	{
 		LoadBinaryIntoBigNum(InPrivateExponent.GetData(), InPrivateExponent.Num(), PrivateExponent);
-		Key->Type = FRSA::FKey::EType::Public;
 	}
 
 	LoadBinaryIntoBigNum(InModulus.GetData(), InModulus.Num(), Modulus);
@@ -130,13 +120,13 @@ bool FRSA::Encrypt(EKeyType InKeyType, const uint8* InSource, int32 InSourceSize
 	
 	switch (InKeyType)
 	{
-	case FRSA::EKeyType::Public:
+	case EKeyType::Public:
 	{
 		NumEncryptedBytes = RSA_public_encrypt(InSourceSizeInBytes, InSource, OutDestination.GetData(), InKey->RSAKey, RSA_PKCS1_PADDING);
 		break;
 	}
 
-	case FRSA::EKeyType::Private:
+	case EKeyType::Private:
 	{
 		NumEncryptedBytes = RSA_private_encrypt(InSourceSizeInBytes, InSource, OutDestination.GetData(), InKey->RSAKey, RSA_PKCS1_PADDING);
 		break;
@@ -160,22 +150,16 @@ bool FRSA::Decrypt(EKeyType InKeyType, const TArray<uint8>& InSource, uint8* Out
 	{
 		int NumDecryptedBytes = 0;
 
-		switch (InKey->Type)
+		switch (InKeyType)
 		{
-		case FKey::EType::Public:
+		case EKeyType::Public:
 		{
-#if USE_LEGACY_OPENSSL
-			FScopeLock Lock(&InKey->Lock);
-#endif
 			NumDecryptedBytes = RSA_public_decrypt(InSource.Num(), InSource.GetData(), OutDestination, InKey->RSAKey, RSA_PKCS1_PADDING);
 			break;
 		}
 
-		case FKey::EType::Private:
+		case EKeyType::Private:
 		{
-#if USE_LEGACY_OPENSSL
-			FScopeLock Lock(&InKey->Lock);
-#endif
 			NumDecryptedBytes = RSA_private_decrypt(InSource.Num(), InSource.GetData(), OutDestination, InKey->RSAKey, RSA_PKCS1_PADDING);
 			break;
 		}
