PlatformCrypto based RSA support for the engine
 - RSA lib now redirects through a modular feature interface to find a crypto provider

diff --git a/Engine/Plugins/Experimental/PlatformCrypto/PlatformCrypto.uplugin b/Engine/Plugins/Experimental/PlatformCrypto/PlatformCrypto.uplugin
--- a/Engine/Plugins/Experimental/PlatformCrypto/PlatformCrypto.uplugin
+++ b/Engine/Plugins/Experimental/PlatformCrypto/PlatformCrypto.uplugin
@@ -1,82 +1,83 @@
 {
-	"FileVersion" : 3,
-	"Version" : 1,
-	"VersionName" : "1.0",
-	"FriendlyName" : "Platform Cryptography Plugin",
-	"Description" : "Exposes a unified API for cryptography functionality provided by the platform, if available. Otherwise, interfaces with OpenSSL.",
-	"Category" : "Misc",
-	"CreatedBy" : "Epic Games, Inc.",
-	"CreatedByURL" : "http://epicgames.com",
-	"DocsURL" : "",
-	"MarketplaceURL" : "",
-	"SupportURL" : "",
-	"EnabledByDefault" : false,
-	"CanContainContent" : false,
-	"IsBetaVersion" : false,
-	"Installed" : false,
-	"Modules" :
-	[
+	"FileVersion": 3,
+	"Version": 1,
+	"VersionName": "1.0",
+	"FriendlyName": "Platform Cryptography Plugin",
+	"Description": "Exposes a unified API for cryptography functionality provided by the platform, if available. Otherwise, interfaces with OpenSSL.",
+	"Category": "Misc",
+	"CreatedBy": "Epic Games, Inc.",
+	"CreatedByURL": "http://epicgames.com",
+	"DocsURL": "",
+	"MarketplaceURL": "",
+	"SupportURL": "",
+	"EnabledByDefault": true,
+	"CanContainContent": false,
+	"IsBetaVersion": false,
+	"Installed": false,
+	"SupportedPrograms": [ "UnrealPak" ],
+	"Modules": [
 		{
 			"Name": "PlatformCrypto",
-			"Type": "Runtime",
-			"LoadingPhase": "Default",
-            "WhitelistPlatforms": [
-              "Android",
-              "IOS",
-              "Mac",
-              "Win32",
-              "Win64",
-              "Linux",
-              "PS4",
-              "XboxOne",
-              "Switch"
-            ],
-            "BlacklistPlatforms": [
-              "TVOS"
-            ]
+			"Type": "RuntimeAndProgram",
+			"LoadingPhase": "EarliestPossible",
+			"WhitelistPlatforms": [
+				"Android",
+				"IOS",
+				"Mac",
+				"Win32",
+				"Win64",
+				"Linux",
+				"PS4",
+				"XboxOne",
+				"Switch"
+			],
+			"BlacklistPlatforms": [
+				"TVOS"
+			],
+			"WhitelistPrograms": [ "UnrealPak" ]
 		},
 		{
 			"Name": "PlatformCryptoOpenSSL",
-			"Type": "Runtime",
-			"LoadingPhase": "Default",
-            "WhitelistPlatforms": [
-              "Android",
-              "IOS",
-              "Mac",
-              "Win32",
-              "Win64",
-              "Linux",
-              "PS4"
-            ],
-            "BlacklistPlatforms": [
-              "TVOS"
-            ]
+			"Type": "RuntimeAndProgram",
+			"LoadingPhase": "EarliestPossible",
+			"WhitelistPlatforms": [
+				"Android",
+				"IOS",
+				"Mac",
+				"Win32",
+				"Win64",
+				"Linux",
+				"PS4"
+			],
+			"BlacklistPlatforms": [
+				"TVOS"
+			]
 		},
 		{
-			"Name" : "PlatformCryptoBCrypt",
-			"Type" : "Runtime",
-			"LoadingPhase" : "Default",
-            "WhitelistPlatforms": [
-              "XboxOne"
-            ],
-            "BlacklistPlatforms": [
-              "Android",
-              "IOS",
-              "TVOS"
-            ]
+			"Name": "PlatformCryptoBCrypt",
+			"Type": "RuntimeAndProgram",
+			"LoadingPhase": "EarliestPossible",
+			"WhitelistPlatforms": [
+				"XboxOne"
+			],
+			"BlacklistPlatforms": [
+				"Android",
+				"IOS",
+				"TVOS"
+			]
 		},
 		{
 			"Name": "PlatformCryptoSwitch",
-			"Type": "Runtime",
-			"LoadingPhase": "Default",
-            "WhitelistPlatforms": [
-              "Switch"
-            ],
-            "BlacklistPlatforms": [
-              "Android",
-              "IOS",
-              "TVOS"
-            ]
+			"Type": "RuntimeAndProgram",
+			"LoadingPhase": "EarliestPossible",
+			"WhitelistPlatforms": [
+				"Switch"
+			],
+			"BlacklistPlatforms": [
+				"Android",
+				"IOS",
+				"TVOS"
+			]
 		}
 	]
 }
\ No newline at end of file
diff --git a/Engine/Plugins/Experimental/PlatformCrypto/Source/PlatformCrypto/Private/PlatformCrypto.cpp b/Engine/Plugins/Experimental/PlatformCrypto/Source/PlatformCrypto/Private/PlatformCrypto.cpp
--- a/Engine/Plugins/Experimental/PlatformCrypto/Source/PlatformCrypto/Private/PlatformCrypto.cpp
+++ b/Engine/Plugins/Experimental/PlatformCrypto/Source/PlatformCrypto/Private/PlatformCrypto.cpp
@@ -3,30 +3,87 @@
 #include "CoreMinimal.h"
 #include "Modules/ModuleManager.h"
 #include "IPlatformCrypto.h"
-
 #include "PlatformCryptoIncludes.h"
+#include "Features/IModularFeatures.h"
+#include "Misc/IEngineCrypto.h"
 
-class FPlatformCrypto : public IPlatformCrypto
+class FPlatformCryptoModularFeature : public IEngineCrypto
 {
-	/** IModuleInterface implementation */
-	virtual void StartupModule() override;
-	virtual void ShutdownModule() override;
-};
+public:
 
-IMPLEMENT_MODULE( FPlatformCrypto, PlatformCrypto )
+	FPlatformCryptoModularFeature()
+	{
+		IModularFeatures::Get().RegisterModularFeature(IEngineCrypto::GetFeatureName(), this);
+	}
 
+	virtual ~FPlatformCryptoModularFeature()
+	{
+		Context.Reset();
+		IModularFeatures::Get().UnregisterModularFeature(IEngineCrypto::GetFeatureName(), this);
+	}
+	
+	/** IEngineCrypto implementation */
+	virtual FRSAKeyHandle CreateRSAKey(const TArrayView<const uint8> InPublicExponent, const TArrayView<const uint8> InPrivateExponent, const TArrayView<const uint8> InModulus) override
+	{
+		return GetContext()->CreateKey_RSA(InPublicExponent, InPrivateExponent, InModulus);
+	}
 
+	virtual void DestroyRSAKey(FRSAKeyHandle InKey) override
+	{
+		GetContext()->DestroyKey_RSA(InKey);
+	}
 
-void FPlatformCrypto::StartupModule()
-{
-}
+	virtual int32 GetKeySize(FRSAKeyHandle InKey) override
+	{
+		return GetContext()->GetKeySize_RSA(InKey);
+	}
 
+	virtual int32 GetMaxDataSize(FRSAKeyHandle InKey) override
+	{
+		return GetContext()->GetMaxDataSize_RSA(InKey);
+	}
 
-void FPlatformCrypto::ShutdownModule()
-{
-}
+	virtual int32 EncryptPublic(const TArrayView<const uint8> InSource, TArray<uint8>& OutDestination, FRSAKeyHandle InKey) override
+	{
+		return GetContext()->EncryptPublic_RSA(InSource, OutDestination, InKey);
+	}
+
+	virtual int32 EncryptPrivate(const TArrayView<const uint8> InSource, TArray<uint8>& OutDestination, FRSAKeyHandle InKey) override
+	{
+		return GetContext()->EncryptPrivate_RSA(InSource, OutDestination, InKey);
+	}
+
+	virtual int32 DecryptPublic(const TArrayView<const uint8> InSource, TArray<uint8>& OutDestination, FRSAKeyHandle InKey) override
+	{
+		return GetContext()->DecryptPublic_RSA(InSource, OutDestination, InKey);
+	}
+
+	virtual int32 DecryptPrivate(const TArrayView<const uint8> InSource, TArray<uint8>& OutDestination, FRSAKeyHandle InKey) override
+	{
+		return GetContext()->DecryptPrivate_RSA(InSource, OutDestination, InKey);
+	}
+
+private:
+
+	TUniquePtr<FEncryptionContext>& GetContext()
+	{
+		if (!Context.IsValid())
+		{
+			Context = IPlatformCrypto::Get().CreateContext();
+		}
+
+		return Context;
+	}
+
+	/** Content used by IEngineCrypto implementation */
+	TUniquePtr<FEncryptionContext> Context;
+};
+
+FPlatformCryptoModularFeature GPlatformCryptoModularFeature;
+
+IMPLEMENT_MODULE(FDefaultModuleImpl, PlatformCrypto)
 
 TUniquePtr<FEncryptionContext> IPlatformCrypto::CreateContext()
 {
 	return MakeUnique<FEncryptionContext>();
-}
+}
\ No newline at end of file
diff --git a/Engine/Plugins/Experimental/PlatformCrypto/Source/PlatformCrypto/Public/PlatformCryptoTypes.h b/Engine/Plugins/Experimental/PlatformCrypto/Source/PlatformCrypto/Public/PlatformCryptoTypes.h
--- a/Engine/Plugins/Experimental/PlatformCrypto/Source/PlatformCrypto/Public/PlatformCryptoTypes.h
+++ b/Engine/Plugins/Experimental/PlatformCrypto/Source/PlatformCrypto/Public/PlatformCryptoTypes.h
@@ -4,6 +4,7 @@
 
 #include "CoreMinimal.h"
 #include "Stats/Stats.h"
+#include "Misc/IEngineCrypto.h"
 
 /**
  * Generic result type for cryptographic functions.
@@ -14,11 +15,6 @@ enum class EPlatformCryptoResult
 	Failure
 };
 
-/**
- * Generic handle to an RSA key
- */
-typedef void* TPlatformCryptoRSAKey;
-
 /**
  * Stat group for implementations to use.
  */
diff --git a/Engine/Plugins/Experimental/PlatformCrypto/Source/PlatformCryptoOpenSSL/Private/EncryptionContextOpenSSL.cpp b/Engine/Plugins/Experimental/PlatformCrypto/Source/PlatformCryptoOpenSSL/Private/EncryptionContextOpenSSL.cpp
--- a/Engine/Plugins/Experimental/PlatformCrypto/Source/PlatformCryptoOpenSSL/Private/EncryptionContextOpenSSL.cpp
+++ b/Engine/Plugins/Experimental/PlatformCrypto/Source/PlatformCryptoOpenSSL/Private/EncryptionContextOpenSSL.cpp
@@ -193,7 +193,7 @@ static void LoadBinaryIntoBigNum(const uint8* InData, int64 InDataSize, BIGNUM*
 #endif
 }
 
-TPlatformCryptoRSAKey FEncryptionContextOpenSSL::CreateKey_RSA(const TArrayView<const uint8> PublicExponent, const TArrayView<const uint8> PrivateExponent, const TArrayView<const uint8> Modulus)
+FRSAKeyHandle FEncryptionContextOpenSSL::CreateKey_RSA(const TArrayView<const uint8> PublicExponent, const TArrayView<const uint8> PrivateExponent, const TArrayView<const uint8> Modulus)
 {
 	RSA* NewKey = RSA_new();
 
@@ -223,22 +223,22 @@ TPlatformCryptoRSAKey FEncryptionContextOpenSSL::CreateKey_RSA(const TArrayView<
 	return NewKey;
 }
 
-void FEncryptionContextOpenSSL::DestroyKey_RSA(TPlatformCryptoRSAKey Key)
+void FEncryptionContextOpenSSL::DestroyKey_RSA(FRSAKeyHandle Key)
 {
 	RSA_free((RSA*)Key);
 }
 
-int32 FEncryptionContextOpenSSL::GetKeySize_RSA(TPlatformCryptoRSAKey Key)
+int32 FEncryptionContextOpenSSL::GetKeySize_RSA(FRSAKeyHandle Key)
 {
 	return RSA_size((RSA*)Key);
 }
 
-int32 FEncryptionContextOpenSSL::GetMaxDataSize_RSA(TPlatformCryptoRSAKey Key)
+int32 FEncryptionContextOpenSSL::GetMaxDataSize_RSA(FRSAKeyHandle Key)
 {
 	return (GetKeySize_RSA(Key)) - RSA_PKCS1_PADDING_SIZE;
 }
 
-int32 FEncryptionContextOpenSSL::EncryptPublic_RSA(TArrayView<const uint8> Source, TArray<uint8>& Dest, TPlatformCryptoRSAKey Key)
+int32 FEncryptionContextOpenSSL::EncryptPublic_RSA(TArrayView<const uint8> Source, TArray<uint8>& Dest, FRSAKeyHandle Key)
 {
 	Dest.SetNum(GetKeySize_RSA(Key));
 	int NumEncryptedBytes = RSA_public_encrypt(Source.Num(), Source.GetData(), Dest.GetData(), (RSA*)Key, RSA_PKCS1_PADDING);
@@ -249,7 +249,7 @@ int32 FEncryptionContextOpenSSL::EncryptPublic_RSA(TArrayView<const uint8> Sourc
 	return NumEncryptedBytes;
 }
 
-int32 FEncryptionContextOpenSSL::EncryptPrivate_RSA(TArrayView<const uint8> Source, TArray<uint8>& Dest, TPlatformCryptoRSAKey Key)
+int32 FEncryptionContextOpenSSL::EncryptPrivate_RSA(TArrayView<const uint8> Source, TArray<uint8>& Dest, FRSAKeyHandle Key)
 {
 	Dest.SetNum(GetKeySize_RSA(Key));
 	int NumEncryptedBytes = RSA_private_encrypt(Source.Num(), Source.GetData(), Dest.GetData(), (RSA*)Key, RSA_PKCS1_PADDING);
@@ -260,7 +260,7 @@ int32 FEncryptionContextOpenSSL::EncryptPrivate_RSA(TArrayView<const uint8> Sour
 	return NumEncryptedBytes;
 }
 
-int32 FEncryptionContextOpenSSL::DecryptPublic_RSA(TArrayView<const uint8> Source, TArray<uint8>& Dest, TPlatformCryptoRSAKey Key)
+int32 FEncryptionContextOpenSSL::DecryptPublic_RSA(TArrayView<const uint8> Source, TArray<uint8>& Dest, FRSAKeyHandle Key)
 {
 	Dest.SetNum(GetKeySize_RSA(Key) - RSA_PKCS1_PADDING_SIZE);
 	int NumDecryptedBytes = RSA_public_decrypt(Source.Num(), Source.GetData(), Dest.GetData(), (RSA*)Key, RSA_PKCS1_PADDING);
@@ -275,7 +275,7 @@ int32 FEncryptionContextOpenSSL::DecryptPublic_RSA(TArrayView<const uint8> Sourc
 	return NumDecryptedBytes;
 }
 
-int32 FEncryptionContextOpenSSL::DecryptPrivate_RSA(TArrayView<const uint8> Source, TArray<uint8>& Dest, TPlatformCryptoRSAKey Key)
+int32 FEncryptionContextOpenSSL::DecryptPrivate_RSA(TArrayView<const uint8> Source, TArray<uint8>& Dest, FRSAKeyHandle Key)
 {
 	Dest.SetNum(GetKeySize_RSA(Key) - RSA_PKCS1_PADDING_SIZE);
 	int NumDecryptedBytes = RSA_private_decrypt(Source.Num(), Source.GetData(), Dest.GetData(), (RSA*)Key, RSA_PKCS1_PADDING);
diff --git a/Engine/Plugins/Experimental/PlatformCrypto/Source/PlatformCryptoOpenSSL/Public/EncryptionContextOpenSSL.h b/Engine/Plugins/Experimental/PlatformCrypto/Source/PlatformCryptoOpenSSL/Public/EncryptionContextOpenSSL.h
--- a/Engine/Plugins/Experimental/PlatformCrypto/Source/PlatformCryptoOpenSSL/Public/EncryptionContextOpenSSL.h
+++ b/Engine/Plugins/Experimental/PlatformCrypto/Source/PlatformCryptoOpenSSL/Public/EncryptionContextOpenSSL.h
@@ -20,14 +20,14 @@ public:
 
 	bool DigestVerify_PS256(const TArrayView<const char> Message, const TArrayView<const uint8> Signature, const TArrayView<const uint8> PKCS1Key);
 
-	TPlatformCryptoRSAKey CreateKey_RSA(const TArrayView<const uint8> PublicExponent, const TArrayView<const uint8> PrivateExponent, const TArrayView<const uint8> Modulus);
-	void DestroyKey_RSA(TPlatformCryptoRSAKey Key);
-	int32 GetKeySize_RSA(TPlatformCryptoRSAKey Key);
-	int32 GetMaxDataSize_RSA(TPlatformCryptoRSAKey Key);
-	int32 EncryptPublic_RSA(TArrayView<const uint8> Source, TArray<uint8>& Dest, TPlatformCryptoRSAKey Key);
-	int32 EncryptPrivate_RSA(TArrayView<const uint8> Source, TArray<uint8>& Dest, TPlatformCryptoRSAKey Key);
-	int32 DecryptPublic_RSA(TArrayView<const uint8> Source, TArray<uint8>& Dest, TPlatformCryptoRSAKey Key);
-	int32 DecryptPrivate_RSA(TArrayView<const uint8> Source, TArray<uint8>& Dest, TPlatformCryptoRSAKey Key);
+	FRSAKeyHandle CreateKey_RSA(const TArrayView<const uint8> PublicExponent, const TArrayView<const uint8> PrivateExponent, const TArrayView<const uint8> Modulus);
+	void DestroyKey_RSA(FRSAKeyHandle Key);
+	int32 GetKeySize_RSA(FRSAKeyHandle Key);
+	int32 GetMaxDataSize_RSA(FRSAKeyHandle Key);
+	int32 EncryptPublic_RSA(TArrayView<const uint8> Source, TArray<uint8>& Dest, FRSAKeyHandle Key);
+	int32 EncryptPrivate_RSA(TArrayView<const uint8> Source, TArray<uint8>& Dest, FRSAKeyHandle Key);
+	int32 DecryptPublic_RSA(TArrayView<const uint8> Source, TArray<uint8>& Dest, FRSAKeyHandle Key);
+	int32 DecryptPrivate_RSA(TArrayView<const uint8> Source, TArray<uint8>& Dest, FRSAKeyHandle Key);
 };
 
 typedef FEncryptionContextOpenSSL FEncryptionContext;
diff --git a/Engine/Source/Developer/PakFileUtilities/Private/PakFileUtilities.cpp b/Engine/Source/Developer/PakFileUtilities/Private/PakFileUtilities.cpp
--- a/Engine/Source/Developer/PakFileUtilities/Private/PakFileUtilities.cpp
+++ b/Engine/Source/Developer/PakFileUtilities/Private/PakFileUtilities.cpp
@@ -41,7 +41,7 @@ struct FNamedAESKey
 
 struct FKeyChain
 {
-	FRSA::TKeyPtr SigningKey;
+	FRSAKeyHandle SigningKey = InvalidRSAKeyHandle;
 	TMap<FGuid, FNamedAESKey> EncryptionKeys;
 	const FNamedAESKey* MasterEncryptionKey = nullptr;
 };
@@ -1258,7 +1258,7 @@ TEncryptionInt ParseEncryptionIntFromJson(TSharedPtr<FJsonObject> InObj, const T
 	}
 }
 
-FRSA::TKeyPtr ParseRSAKeyFromJson(TSharedPtr<FJsonObject> InObj)
+FRSAKeyHandle ParseRSAKeyFromJson(TSharedPtr<FJsonObject> InObj)
 {
 	TSharedPtr<FJsonObject> PublicKey = InObj->GetObjectField(TEXT("PublicKey"));
 	TSharedPtr<FJsonObject> PrivateKey = InObj->GetObjectField(TEXT("PrivateKey"));
@@ -1347,7 +1347,7 @@ void LoadKeyChainFromFile(const FString& InFilename, FKeyChain& OutCryptoSetting
 
 void LoadKeyChain(const TCHAR* CmdLine, FKeyChain& OutCryptoSettings)
 {
-	OutCryptoSettings.SigningKey.Reset();
+	OutCryptoSettings.SigningKey = InvalidRSAKeyHandle;
 	OutCryptoSettings.EncryptionKeys.Empty();
 
 	// First, try and parse the keys from a supplied crypto key cache file
@@ -1602,7 +1602,7 @@ bool CreatePakFile(const TCHAR* Filename, TArray<FPakInputPair>& FilesToAdd, con
 
 	if (InKeyChain.MasterEncryptionKey)
 	{
-		UE_LOG(LogPakFile, Display, TEXT("Encrypting using key '%s' [%s]"), *InKeyChain.MasterEncryptionKey->Name, *InKeyChain.MasterEncryptionKey->Guid.ToString());
+		UE_LOG(LogPakFile, Display, TEXT("Using encryption key '%s' [%s]"), *InKeyChain.MasterEncryptionKey->Name, *InKeyChain.MasterEncryptionKey->Guid.ToString());
 	}
 
 	TArray<FPakEntryPair> Index;
@@ -3824,7 +3824,7 @@ bool Repack(const FString& InputPakFile, const FString& OutputPakFile, const FPa
 		FString TempOutputPakFile = FPaths::CreateTempFilename(*FPaths::GetPath(OutputPakFile), *FPaths::GetCleanFilename(OutputPakFile));
 
 		FPakCommandLineParameters ModifiedCmdLineParameters = CmdLineParameters;
-		ModifiedCmdLineParameters.bSign = bAnySigned && InKeyChain.SigningKey.IsValid();
+		ModifiedCmdLineParameters.bSign = bAnySigned && (InKeyChain.SigningKey != InvalidRSAKeyHandle);
 		
 		FKeyChain ModifiedKeyChain = InKeyChain;
 		ModifiedKeyChain.MasterEncryptionKey = InKeyChain.EncryptionKeys.Find(EncryptionKeys.Num() ? EncryptionKeys[0] : FGuid());
diff --git a/Engine/Source/Developer/PakFileUtilities/Private/SignedArchiveWriter.cpp b/Engine/Source/Developer/PakFileUtilities/Private/SignedArchiveWriter.cpp
--- a/Engine/Source/Developer/PakFileUtilities/Private/SignedArchiveWriter.cpp
+++ b/Engine/Source/Developer/PakFileUtilities/Private/SignedArchiveWriter.cpp
@@ -5,7 +5,7 @@
 #include "Misc/SecureHash.h"
 #include "HAL/FileManager.h"
 
-FSignedArchiveWriter::FSignedArchiveWriter(FArchive& InPak, const FString& InPakFilename, FRSA::TKeyPtr InSigningKey)
+FSignedArchiveWriter::FSignedArchiveWriter(FArchive& InPak, const FString& InPakFilename, const FRSAKeyHandle InSigningKey)
 : BufferArchive(Buffer)
 	, PakWriter(InPak)
 	, PakSignaturesFilename(FPaths::ChangeExtension(InPakFilename, TEXT("sig")))
diff --git a/Engine/Source/Developer/PakFileUtilities/Private/SignedArchiveWriter.h b/Engine/Source/Developer/PakFileUtilities/Private/SignedArchiveWriter.h
--- a/Engine/Source/Developer/PakFileUtilities/Private/SignedArchiveWriter.h
+++ b/Engine/Source/Developer/PakFileUtilities/Private/SignedArchiveWriter.h
@@ -25,7 +25,7 @@ class FSignedArchiveWriter : public FArchive
 	/** Data size (excluding signatures) */
 	int64 PakSize;
 	/** Signing key */
-	FRSA::TKeyPtr SigningKey;
+	const FRSAKeyHandle SigningKey;
 	/** Hashes */
 	TArray<TPakChunkHash> ChunkHashes;
 
@@ -36,7 +36,7 @@ class FSignedArchiveWriter : public FArchive
 
 public:
 
-	FSignedArchiveWriter(FArchive& InPak, const FString& InPakFilename, FRSA::TKeyPtr InSigningKey);
+	FSignedArchiveWriter(FArchive& InPak, const FString& InPakFilename, const FRSAKeyHandle InSigningKey);
 	virtual ~FSignedArchiveWriter();
 
 	// FArchive interface
diff --git a/Engine/Source/Runtime/Core/Public/Misc/IEngineCrypto.h b/Engine/Source/Runtime/Core/Public/Misc/IEngineCrypto.h
new file mode 100644
--- /dev/null
+++ b/Engine/Source/Runtime/Core/Public/Misc/IEngineCrypto.h
@@ -0,0 +1,59 @@
+#pragma once
+
+#include "CoreTypes.h"
+#include "Features/IModularFeatures.h"
+
+typedef void* FRSAKeyHandle;
+static const FRSAKeyHandle InvalidRSAKeyHandle = nullptr;
+
+struct IEngineCrypto : public IModularFeature
+{
+	/**
+	* Get the name of this modular feature
+	*/
+	static FORCEINLINE FName GetFeatureName()
+	{
+		static const FName Name(TEXT("EngineCryptoFeature"));
+		return Name;
+	}
+
+	/** 
+	* Create a new RSA key from the given little-endian exponents and modulus
+	*/
+	virtual FRSAKeyHandle CreateRSAKey(const TArrayView<const uint8> InPublicExponent, const TArrayView<const uint8> InPrivateExponent, const TArrayView<const uint8> InModulus) = 0;
+
+	/**
+	* Destroy the given RSA key
+	*/
+	virtual void DestroyRSAKey(FRSAKeyHandle InKey) = 0;
+
+	/**
+	* Get the size of bytes of the given RSA key
+	*/
+	virtual int32 GetKeySize(FRSAKeyHandle InKey) = 0;
+
+	/** 
+	* Get the maximum amount of data that can be encrypted using the given key, taking into account minimum padding requirements
+	*/
+	virtual int32 GetMaxDataSize(FRSAKeyHandle InKey) = 0;
+
+	/**
+	* Encrypt the supplied byte data using the given public key
+	*/
+	virtual int32 EncryptPublic(const TArrayView<const uint8> InSource, TArray<uint8>& OutDestination, FRSAKeyHandle InKey) = 0;
+
+	/**
+	 * Encrypt the supplied byte data using the given private key
+	 */
+	virtual int32 EncryptPrivate(const TArrayView<const uint8> InSource, TArray<uint8>& OutDestination, FRSAKeyHandle InKey) = 0;
+
+	/**
+	 * Decrypt the supplied byte data using the given public key
+	 */
+	virtual int32 DecryptPublic(const TArrayView<const uint8> InSource, TArray<uint8>& OutDestination, FRSAKeyHandle InKey) = 0;
+
+	/**
+	 * Encrypt the supplied byte data using the given private key
+	 */
+	virtual int32 DecryptPrivate(const TArrayView<const uint8> InSource, TArray<uint8>& OutDestination, FRSAKeyHandle InKey) = 0;
+};
\ No newline at end of file
diff --git a/Engine/Source/Runtime/PakFile/Private/IPlatformFilePak.cpp b/Engine/Source/Runtime/PakFile/Private/IPlatformFilePak.cpp
--- a/Engine/Source/Runtime/PakFile/Private/IPlatformFilePak.cpp
+++ b/Engine/Source/Runtime/PakFile/Private/IPlatformFilePak.cpp
@@ -220,13 +220,21 @@ void FPakPlatformFile::GetPakEncryptionKey(FAES::FAESKey& OutKey, const FGuid& I
 	}
 }
 
-TSharedPtr<FRSA::FKey, ESPMode::ThreadSafe> FPakPlatformFile::GetPakSigningKey()
+TMap<FName, TSharedPtr<const struct FPakSignatureFile, ESPMode::ThreadSafe>> FPakPlatformFile::PakSignatureFileCache;
+FCriticalSection FPakPlatformFile::PakSignatureFileCacheLock;
+TSharedPtr<const struct FPakSignatureFile, ESPMode::ThreadSafe> FPakPlatformFile::GetPakSignatureFile(const TCHAR* InFilename)
 {
-	static TSharedPtr<FRSA::FKey, ESPMode::ThreadSafe> Key;
-	static FCriticalSection Lock;
-	Lock.Lock();
+	FScopeLock Lock(&PakSignatureFileCacheLock);
+
+	FName FilenameFName(InFilename);
+	if (const TSharedPtr<const struct FPakSignatureFile, ESPMode::ThreadSafe>* SignaturesFile = PakSignatureFileCache.Find(FilenameFName))
+	{
+		return *SignaturesFile;
+	}
 
-	if (!Key.IsValid())
+	static FRSAKeyHandle PublicKey = InvalidRSAKeyHandle;
+	static bool bInitializedPublicKey = false;
+	if (!bInitializedPublicKey)
 	{
 		FCoreDelegates::FPakSigningKeysDelegate& Delegate = FCoreDelegates::GetPakSigningKeysDelegate();
 		if (Delegate.IsBound())
@@ -234,44 +242,37 @@ TSharedPtr<FRSA::FKey, ESPMode::ThreadSafe> FPakPlatformFile::GetPakSigningKey()
 			TArray<uint8> Exponent;
 			TArray<uint8> Modulus;
 			Delegate.Execute(Exponent, Modulus);
-			Key = FRSA::CreateKey(Exponent, TArray<uint8>(), Modulus);
+			PublicKey = FRSA::CreateKey(Exponent, TArray<uint8>(), Modulus);
 		}
+		bInitializedPublicKey = true;
 	}
-	
-	Lock.Unlock();
-	return Key;
-}
 
-TMap<FName, TSharedPtr<const struct FPakSignatureFile, ESPMode::ThreadSafe>> FPakPlatformFile::PakSignatureFileCache;
-FCriticalSection FPakPlatformFile::PakSignatureFileCacheLock;
-TSharedPtr<const struct FPakSignatureFile, ESPMode::ThreadSafe> FPakPlatformFile::GetPakSignatureFile(const TCHAR* InFilename)
-{
-	FScopeLock Lock(&PakSignatureFileCacheLock);
-	FName FilenameFName(InFilename);
-	if (const TSharedPtr<const struct FPakSignatureFile, ESPMode::ThreadSafe>* SignaturesFile = PakSignatureFileCache.Find(FilenameFName))
-	{
-		return *SignaturesFile;
-	}
+	TSharedPtr<FPakSignatureFile, ESPMode::ThreadSafe> SignaturesFile;
 
-	TSharedPtr<FPakSignatureFile, ESPMode::ThreadSafe> SignaturesFile = MakeShared<FPakSignatureFile, ESPMode::ThreadSafe>();
-	FString SignaturesFilename = FPaths::ChangeExtension(InFilename, TEXT("sig"));
-	FArchive* Reader = IFileManager::Get().CreateFileReader(*SignaturesFilename);
-	if (Reader != nullptr)
+	if (PublicKey != InvalidRSAKeyHandle)
 	{
-		SignaturesFile->Serialize(*Reader);
-		delete Reader;
+		FString SignaturesFilename = FPaths::ChangeExtension(InFilename, TEXT("sig"));
+		FArchive* Reader = IFileManager::Get().CreateFileReader(*SignaturesFilename);
+		if (Reader != nullptr)
+		{
+			SignaturesFile = MakeShared<FPakSignatureFile, ESPMode::ThreadSafe>();
+			SignaturesFile->Serialize(*Reader);
+			delete Reader;
+
+			if (!SignaturesFile->DecryptSignatureAndValidate(PublicKey, InFilename))
+			{
+				// We don't need to act on this failure as the decrypt function will already have dumped out log messages
+				// and fired the signature check fail handler
+				SignaturesFile.Reset();
+			}
 
-		if (!SignaturesFile->DecryptSignatureAndValidate(InFilename))
+			PakSignatureFileCache.Add(FilenameFName, SignaturesFile);
+		}
+		else
 		{
-			SignaturesFile.Reset();
+			UE_LOG(LogPakFile, Warning, TEXT("Couldn't find pak signature file '%s'"), InFilename);
+			FPakPlatformFile::GetPakMasterSignatureTableCheckFailureHandler().Broadcast(InFilename);
 		}
-
-		PakSignatureFileCache.Add(FilenameFName, SignaturesFile);
-	}
-	else
-	{
-		UE_LOG(LogPakFile, Warning, TEXT("Couldn't find pak signature file '%s'"), InFilename);
-		FPakPlatformFile::GetPakMasterSignatureTableCheckFailureHandler().Broadcast(InFilename);
 	}
 
 	return SignaturesFile;
@@ -1196,16 +1197,17 @@ class FPakPrecacher
 	uint32 Loads;
 	uint32 Frees;
 	uint64 LoadSize;
-	FRSA::TKeyPtr SigningKey;
 	EAsyncIOPriorityAndFlags AsyncMinPriority;
 	FCriticalSection SetAsyncMinimumPriorityScopeLock;
+	bool bEnableSignatureChecks;
+
 public:
 
-	static void Init(IPlatformFile* InLowerLevel, FRSA::TKeyPtr InSigningKey)
+	static void Init(IPlatformFile* InLowerLevel, bool bInEnableSignatureChecks)
 	{
 		if (!PakPrecacherSingleton)
 		{
-			verify(!FPlatformAtomics::InterlockedCompareExchangePointer((void**)&PakPrecacherSingleton, new FPakPrecacher(InLowerLevel, InSigningKey), nullptr));
+			verify(!FPlatformAtomics::InterlockedCompareExchangePointer((void**)&PakPrecacherSingleton, new FPakPrecacher(InLowerLevel, bInEnableSignatureChecks), nullptr));
 		}
 		check(PakPrecacherSingleton);
 	}
@@ -1241,7 +1243,7 @@ public:
 		return *PakPrecacherSingleton;
 	}
 
-	FPakPrecacher(IPlatformFile* InLowerLevel, FRSA::TKeyPtr InSigningKey)
+	FPakPrecacher(IPlatformFile* InLowerLevel, bool bInEnableSignatureChecks)
 		: LowerLevel(InLowerLevel)
 		, LastReadRequest(0)
 		, NextUniqueID(1)
@@ -1251,8 +1253,8 @@ public:
 		, Loads(0)
 		, Frees(0)
 		, LoadSize(0)
-		, SigningKey(InSigningKey)
 		, AsyncMinPriority(AIOP_MIN)
+		, bEnableSignatureChecks(bInEnableSignatureChecks)
 	{
 		check(LowerLevel && FPlatformProcess::SupportsMultithreading());
 		GPakCache_MaxRequestsToLowerLevel = FMath::Max(FMath::Min(FPlatformMisc::NumberOfIOWorkerThreadsToSpawn(), GPakCache_MaxRequestsToLowerLevel), 1);
@@ -1290,17 +1292,17 @@ public:
 			UE_LOG(LogPakFile, Log, TEXT("New pak file %s added to pak precacher."), *PakFilename);
 
 			FPakData& Pak = CachedPakData[*PakIndexPtr];
-
-			if (SigningKey.IsValid())
-			{
-				// Load signature data
-				Pak.Signatures = FPakPlatformFile::GetPakSignatureFile(*PakFilename);
+			
+			// Load signature data
+			Pak.Signatures = FPakPlatformFile::GetPakSignatureFile(*PakFilename);
 				
+			if (Pak.Signatures.IsValid())
+			{
 				// We should never get here unless the signature file exists and is validated. The original FPakFile creation
 				// on the main thread would have failed and the pak would never have been mounted otherwise, and then we would
 				// never have issued read requests to the pak precacher.
 				check(Pak.Signatures);
-				
+
 				// Check that we have the correct match between signature and pre-cache granularity
 				int64 NumPakChunks = Align(PakFileSize, FPakInfo::MaxChunkDataSize) / FPakInfo::MaxChunkDataSize;
 				ensure(NumPakChunks == Pak.Signatures->ChunkHashes.Num());
@@ -2123,7 +2125,7 @@ private: // below here we assume CachedFilesScopeLock until we get to the next s
 		FAsyncFileCallBack CallbackFromLower =
 			[this, IndexToFill, bDoCheck](bool bWasCanceled, IAsyncReadRequest* Request)
 		{
-			if (SigningKey.IsValid() && bDoCheck)
+			if (bEnableSignatureChecks && bDoCheck)
 			{
 				StartSignatureCheck(bWasCanceled, Request, IndexToFill);
 			}
@@ -5538,9 +5540,9 @@ bool FPakPlatformFile::Initialize(IPlatformFile* Inner, const TCHAR* CmdLine)
 	GameUserSettingsIniFilename = TEXT("GameUserSettings.ini");
 #endif
 
-	// signed if we have keys, and are not running with fileopenlog (currently results in a deadlock).
-	bSigned = GetPakSigningKey().IsValid() && !FParse::Param(FCommandLine::Get(), TEXT("fileopenlog"));;
-		
+	// Signed if we have keys, and are not running with fileopenlog (currently results in a deadlock).
+	bSigned = FCoreDelegates::GetPakSigningKeysDelegate().IsBound() && !FParse::Param(FCommandLine::Get(), TEXT("fileopenlog"));
+
 	// Find and mount pak files from the specified directories.
 	TArray<FString> PakFolders;
 	GetPakFolders(FCommandLine::Get(), PakFolders);
@@ -5596,7 +5598,7 @@ void FPakPlatformFile::InitializeNewAsyncIO()
 #if !WITH_EDITOR
 	if (FPlatformProcess::SupportsMultithreading() && !FParse::Param(FCommandLine::Get(), TEXT("FileOpenLog")))
 	{
-		FPakPrecacher::Init(LowerLevel, GetPakSigningKey());
+		FPakPrecacher::Init(LowerLevel, FCoreDelegates::GetPakSigningKeysDelegate().IsBound());
 	}
 	else
 #endif
diff --git a/Engine/Source/Runtime/PakFile/Public/IPlatformFilePak.h b/Engine/Source/Runtime/PakFile/Public/IPlatformFilePak.h
--- a/Engine/Source/Runtime/PakFile/Public/IPlatformFilePak.h
+++ b/Engine/Source/Runtime/PakFile/Public/IPlatformFilePak.h
@@ -1548,11 +1548,6 @@ public:
 	*/
 	static void GetPakEncryptionKey(FAES::FAESKey& OutKey, const FGuid& InEncryptionKeyGuid);
 
-	/**
-	* Helper function for accessing pak the signing public key
-	*/
-	static TSharedPtr<FRSA::FKey, ESPMode::ThreadSafe> GetPakSigningKey();
-
 	/**
 	* Load a pak signature file. Validates the contents by comparing a SHA hash of the chunk table against and encrypted version that
 	* is stored within the file. Returns nullptr if the data is missing or fails the signature check. This function also calls
@@ -2415,11 +2410,11 @@ struct FPakSignatureFile
 	/**
 	 * Initialize and hash the CRC list then use the provided private key to encrypt the hash
 	 */
-	void SetChunkHashesAndSign(const TArray<TPakChunkHash>& InChunkHashes, const FRSA::TKeyPtr InKey)
+	void SetChunkHashesAndSign(const TArray<TPakChunkHash>& InChunkHashes, const FRSAKeyHandle InKey)
 	{
 		ChunkHashes = InChunkHashes;
 		DecryptedHash = ComputeCurrentMasterHash();
-		FRSA::Encrypt(FRSA::EKeyType::Private, DecryptedHash.Hash, ARRAY_COUNT(FSHAHash::Hash), EncryptedHash, InKey);
+		FRSA::EncryptPrivate(TArrayView<const uint8>(DecryptedHash.Hash, ARRAY_COUNT(FSHAHash::Hash)), EncryptedHash, InKey);
 	}
 
 	/**
@@ -2446,34 +2441,32 @@ struct FPakSignatureFile
 	/**
 	 * Decrypt the chunk CRCs hash and validate that it matches the current one
 	 */
-	bool DecryptSignatureAndValidate(const FString& InFilename)
-	{
-		TSharedPtr<FRSA::FKey, ESPMode::ThreadSafe> PublicKey = FPakPlatformFile::GetPakSigningKey();
-		return DecryptSignatureAndValidate(PublicKey, InFilename);
-	}
-
-	/**
-	 * Decrypt the chunk CRCs hash and validate that it matches the current one
-	 */
-	bool DecryptSignatureAndValidate(FRSA::TKeyPtr InKey, const FString& InFilename)
+	bool DecryptSignatureAndValidate(const FRSAKeyHandle InKey, const FString& InFilename)
 	{
 		if (Version == EVersion::Invalid)
 		{
 			UE_LOG(LogPakFile, Warning, TEXT("Pak signature file for '%s' was invalid"), *InFilename);
 		}
-		else if (FRSA::Decrypt(FRSA::EKeyType::Public, EncryptedHash, DecryptedHash.Hash, ARRAY_COUNT(FSHAHash::Hash), InKey))
+		else
 		{
-			FSHAHash CurrentHash = ComputeCurrentMasterHash();
-			if (DecryptedHash == CurrentHash)
+			TArray<uint8> Decrypted;
+			int32 BytesDecrypted = FRSA::DecryptPublic(EncryptedHash, Decrypted, InKey);
+			if (BytesDecrypted == ARRAY_COUNT(FSHAHash::Hash))
 			{
-				return true;
+				FSHAHash CurrentHash = ComputeCurrentMasterHash();
+				if (FMemory::Memcmp(Decrypted.GetData(), CurrentHash.Hash, Decrypted.Num()) == 0)
+				{
+					return true;
+				}
+				else
+				{
+					UE_LOG(LogPakFile, Warning, TEXT("Pak signature table validation failed for '%s'! Expected %s, Received %s"), *InFilename, *DecryptedHash.ToString(), *CurrentHash.ToString());
+				}
+			}
+			else
+			{
+				UE_LOG(LogPakFile, Warning, TEXT("Pak signature table validation failed for '%s'! Failed to decrypt signature"), *InFilename);
 			}
-
-			UE_LOG(LogPakFile, Warning, TEXT("Pak signature table validation failed for '%s'! Expected %s, Received %s"), *InFilename, *DecryptedHash.ToString(), *CurrentHash.ToString());
-		}
-		else
-		{
-			UE_LOG(LogPakFile, Warning, TEXT("Pak signature table validation failed for '%s'! Failed to decrypt signature"), *InFilename);
 		}
 
 		FPakPlatformFile::GetPakMasterSignatureTableCheckFailureHandler().Broadcast(InFilename);
diff --git a/Engine/Source/Runtime/RSA/Private/RSA.cpp b/Engine/Source/Runtime/RSA/Private/RSA.cpp
--- a/Engine/Source/Runtime/RSA/Private/RSA.cpp
+++ b/Engine/Source/Runtime/RSA/Private/RSA.cpp
@@ -1,410 +1,93 @@
 // Copyright 1998-2019 Epic Games, Inc. All Rights Reserved.
 
 #include "RSA.h"
-#include "Misc/ScopeLock.h"
-#include "Modules/ModuleManager.h"
-#include "Math/BigInt.h"
+#include "Features/IModularFeatures.h"
 
-#ifndef RSA_USE_OPENSSL
-#define RSA_USE_OPENSSL 0
-#endif
+static FCriticalSection GLock;
 
-#if RSA_USE_OPENSSL
-#if PLATFORM_WINDOWS
-#include "Windows/AllowWindowsPlatformTypes.h"
-#endif
-
-#include <openssl/rsa.h>
-#include <openssl/bn.h>
-#include <openssl/err.h>
-
-#if PLATFORM_WINDOWS
-#include "Windows/HideWindowsPlatformTypes.h"
-#endif
-
-// Some platforms were upgraded to OpenSSL 1.1.1 while the others were left on a previous version. There are some minor differences we have to account for
-// in the older version, so declare a handy define that we can use to gate the code
-#if !defined(OPENSSL_VERSION_NUMBER) || OPENSSL_VERSION_NUMBER < 0x10100000L
-#define USE_LEGACY_OPENSSL 1
-#else
-#define USE_LEGACY_OPENSSL 0
-#endif
-
-struct FRSA::FKey
+IEngineCrypto* GetEngineCrypto()
 {
-	FKey() 
-	: KeySize(0)
-	, KeySizeInBytes(0)
-	, MaxDataSize(0)
-	, RSAKey(RSA_new())
-	{
-	}
-
-	~FKey()
-	{
-		RSA_free(RSAKey);
-	}
-
-	int32 GetKeySizeInBits() const
-	{
-		return KeySize;
-	}
-
-	int32 GetKeySizeInBytes() const
-	{
-		return KeySizeInBytes;
-	}
-
-	int32 GetMaxDataSize() const
-	{
-		return MaxDataSize;
-	}
-
-	int32 KeySize;
-	int32 KeySizeInBytes;
-	int32 MaxDataSize;
-	RSA* RSAKey;
-};
-
-/**
- * Wrapper function for extracting the bytes of an OpenSSL big num in a consistent form
- */
-void LoadBinaryIntoBigNum(const uint8* InData, int64 InDataSize, BIGNUM* InBigNum)
-{
-#if USE_LEGACY_OPENSSL
-	TArray<uint8> Bytes(InData, InDataSize);
-	Algo::Reverse(Bytes);
-	BN_bin2bn(Bytes.GetData(), Bytes.Num(), InBigNum);
-#else
-	BN_lebin2bn(InData, InDataSize, InBigNum);
-#endif
+	static TArray<IEngineCrypto*> Features = IModularFeatures::Get().GetModularFeatureImplementations<IEngineCrypto>(IEngineCrypto::GetFeatureName());
+	checkf(Features.Num() > 0, TEXT("RSA functionality was used but no modular feature was registered to provide it. Please make sure your project has the PlatformCrypto plugin enabled!"));
+	return Features[0];
 }
 
-FRSA::TKeyPtr FRSA::CreateKey(const TArray<uint8>& InPublicExponent, const TArray<uint8>& InPrivateExponent, const TArray<uint8>& InModulus)
+FRSAKeyHandle FRSA::CreateKey(const TArray<uint8>& InPublicExponent, const TArray<uint8>& InPrivateExponent, const TArray<uint8>& InModulus)
 {
-	FRSA::TKeyPtr Key = MakeShared<FRSA::FKey, FRSA::KeyThreadMode>();
-	BIGNUM* PublicExponent = InPublicExponent.Num() > 0 ? BN_new() : nullptr;
-	BIGNUM* PrivateExponent = InPrivateExponent.Num() > 0 ? BN_new() : nullptr;
-	BIGNUM* Modulus = BN_new();
-	
-	if (InPublicExponent.Num())
-	{
-		LoadBinaryIntoBigNum(InPublicExponent.GetData(), InPublicExponent.Num(), PublicExponent);
-	}
-	
-	if (InPrivateExponent.Num())
+	if (IEngineCrypto* EngineCrypto = GetEngineCrypto())
 	{
-		LoadBinaryIntoBigNum(InPrivateExponent.GetData(), InPrivateExponent.Num(), PrivateExponent);
+		return EngineCrypto->CreateRSAKey(InPublicExponent, InPrivateExponent, InModulus);
 	}
 
-	LoadBinaryIntoBigNum(InModulus.GetData(), InModulus.Num(), Modulus);
-#if USE_LEGACY_OPENSSL
-	Key->RSAKey->n = Modulus;
-	Key->RSAKey->e = PublicExponent;
-	Key->RSAKey->d = PrivateExponent;
-#else
-	RSA_set0_key(Key->RSAKey, Modulus, PublicExponent, PrivateExponent);
-#endif
-
-	Key->KeySizeInBytes = RSA_size(Key->RSAKey);
-	Key->KeySize = Key->KeySizeInBytes * 8;
-	Key->MaxDataSize = Key->KeySizeInBytes - RSA_PKCS1_PADDING_SIZE;
-	return Key;
+	return nullptr;
 }
 
-bool FRSA::Encrypt(EKeyType InKeyType, const uint8* InSource, int32 InSourceSizeInBytes, TArray<uint8>& OutDestination, FRSA::TKeyPtr InKey)
+void FRSA::DestroyKey(const FRSAKeyHandle InKey)
 {
-	OutDestination.SetNum(InKey->KeySizeInBytes);
-
-	int NumEncryptedBytes = 0;
-	
-	switch (InKeyType)
-	{
-	case EKeyType::Public:
-	{
-		NumEncryptedBytes = RSA_public_encrypt(InSourceSizeInBytes, InSource, OutDestination.GetData(), InKey->RSAKey, RSA_PKCS1_PADDING);
-		break;
-	}
-
-	case EKeyType::Private:
-	{
-		NumEncryptedBytes = RSA_private_encrypt(InSourceSizeInBytes, InSource, OutDestination.GetData(), InKey->RSAKey, RSA_PKCS1_PADDING);
-		break;
-	}
-	}
-	
-	if (NumEncryptedBytes == InKey->KeySizeInBytes)
-	{
-		return true;
-	}
-	else
+	if (IEngineCrypto* EngineCrypto = GetEngineCrypto())
 	{
-		OutDestination.Empty();
-		return false;
+		EngineCrypto->DestroyRSAKey(InKey);
 	}
 }
 
-bool FRSA::Decrypt(EKeyType InKeyType, const TArray<uint8>& InSource, uint8* OutDestination, int32 OutDestinationSizeInBytes, FRSA::TKeyPtr InKey)
+int32 FRSA::GetKeySize(const FRSAKeyHandle InKey)
 {
-	if (InSource.Num() == InKey->GetKeySizeInBytes() && OutDestinationSizeInBytes <= InKey->GetKeySizeInBytes())
+	if (IEngineCrypto* EngineCrypto = GetEngineCrypto())
 	{
-		int NumDecryptedBytes = 0;
-
-		switch (InKeyType)
-		{
-		case EKeyType::Public:
-		{
-			NumDecryptedBytes = RSA_public_decrypt(InSource.Num(), InSource.GetData(), OutDestination, InKey->RSAKey, RSA_PKCS1_PADDING);
-			break;
-		}
-
-		case EKeyType::Private:
-		{
-			NumDecryptedBytes = RSA_private_decrypt(InSource.Num(), InSource.GetData(), OutDestination, InKey->RSAKey, RSA_PKCS1_PADDING);
-			break;
-		}
-		}
-
-		if (NumDecryptedBytes >= 0 && NumDecryptedBytes <= OutDestinationSizeInBytes)
-		{
-			return true;
-		}
-		else
-		{
-			long Error = ERR_get_error();
-			const char* Message = ERR_error_string(Error, nullptr);
-			return false;
-		}
+		return EngineCrypto->GetKeySize(InKey);
 	}
 
-	return false;
+	return 0;
 }
 
-#else // !RSA_USE_OPENSSL
-
-//#pragma message ("Using TBigInt for RSA")
-#include "Math/RandomStream.h"
-
-struct FRSA::FKey
+int32 FRSA::GetMaxDataSize(const FRSAKeyHandle InKey)
 {
-	virtual int32 GetKeySizeInBits() const = 0;
-	virtual int32 GetKeySizeInBytes() const = 0;
-	virtual int32 GetMaxDataSize() const = 0;
-	virtual bool Encrypt(FRSA::EKeyType InKeyType, const uint8* InSource, int32 InSourceSizeInBytes, TArray<uint8>& OutDestination) const = 0;
-	virtual bool Decrypt(FRSA::EKeyType InKeyType, const TArray<uint8>& InSource, uint8* OutDestination, int32 OutDestinationSizeInBytes) const = 0;
-};
-
-template <int32 KEY_SIZE>
-struct FFixedKey : public FRSA::FKey
-{
-	typedef TBigInt<KEY_SIZE * 2, false> TInternalBigInt;
-	TInternalBigInt PublicExponent;
-	TInternalBigInt PrivateExponent;
-	TInternalBigInt Modulus;
-
-	FFixedKey(const TArray<uint8>& InPublicExponent, const TArray<uint8>& InPrivateExponent, const TArray<uint8>& InModulus)
-	{
-		if (InPublicExponent.Num())
-		{
-			PublicExponent = TInternalBigInt(InPublicExponent.GetData(), InPublicExponent.Num());
-		}
-
-		if (InPrivateExponent.Num())
-		{
-			PrivateExponent = TInternalBigInt(InPrivateExponent.GetData(), InPrivateExponent.Num());
-		}
-
-		if (InModulus.Num())
-		{
-			Modulus = TInternalBigInt(InModulus.GetData(), InModulus.Num());
-		}
-	}
-
-	int32 GetKeySizeInBits() const override
-	{
-		return KEY_SIZE;
-	}
-
-	int32 GetKeySizeInBytes() const override
-	{
-		return KEY_SIZE / 8;
-	}
-
-	/** 
-		* Return the maximum amount of data that can be encrypted within the key, as you would if you were using a proper 
-		* RSA padding scheme. Because this is a legacy system and should be deprecated soon, we're just going to use random byte
-		* padding, but we'll reflect the correct RSA rules anyway
-		*/
-	int32 GetMaxDataSize() const override
-	{
-		return GetKeySizeInBytes() - 11;
-	}
-
-	bool Encrypt(FRSA::EKeyType InKeyType, const uint8* InSource, int32 InSourceSizeInBytes, TArray<uint8>& OutDestination) const override
+	if (IEngineCrypto* EngineCrypto = GetEngineCrypto())
 	{
-		if (InSourceSizeInBytes <= GetMaxDataSize())
-		{
-			TInternalBigInt Source(InSource, InSourceSizeInBytes);
-			int32 PaddingBytesRequired = GetKeySizeInBytes() - InSourceSizeInBytes;
-			uint8* PaddingStart = (uint8*)Source.GetBits() + InSourceSizeInBytes;
-			// Not perfect, but an improvement on before, and soon to be defuncted anyway
-			for (int32 PaddingIndex = 0; PaddingIndex < PaddingBytesRequired; ++PaddingIndex)
-			{
-				PaddingStart[PaddingIndex] = (uint8)FMath::RandRange(0, 255);
-			}
-			PaddingStart[PaddingBytesRequired - 1] &= 0x3f; // keep top two bits clear because this seems to cause the decryption to fail. TODO: Find out why... presumably some bug in TBigInt
-			TInternalBigInt Result = FEncryption::ModularPow(Source, InKeyType == FRSA::EKeyType::Public ? PublicExponent : PrivateExponent, Modulus);
-			TInternalBigInt Test = FEncryption::ModularPow(Result, InKeyType == FRSA::EKeyType::Public ? PrivateExponent : PublicExponent, Modulus);
-			check(Source == Test);
-			OutDestination.SetNum(GetKeySizeInBytes());
-			FMemory::Memcpy(OutDestination.GetData(), Result.GetBits(), GetKeySizeInBytes());
-		}
-		else
-		{
-			OutDestination.Empty();
-		}
-
-		return true;
+		return EngineCrypto->GetMaxDataSize(InKey);
 	}
 
-	bool Decrypt(FRSA::EKeyType InKeyType, const TArray<uint8>& InSource, uint8* OutDestination, int32 OutDestinationSizeInBytes) const override
-	{
-		if (InSource.Num() == GetKeySizeInBytes() && OutDestinationSizeInBytes <= GetKeySizeInBytes())
-		{
-			TInternalBigInt Source(InSource.GetData(), InSource.Num());
-			TInternalBigInt Result = FEncryption::ModularPow(Source, InKeyType == FRSA::EKeyType::Public ? PublicExponent : PrivateExponent, Modulus);
-			FMemory::Memcpy(OutDestination, Result.GetBits(), OutDestinationSizeInBytes);
-			return true;
-		}
-		else
-		{
-			FMemory::Memset(OutDestination, 0, OutDestinationSizeInBytes);
-			return false;
-		}
-	}
-};
+	return 0;
+}
 
-int32 NumElementsIgnoringTrailingZeroes(const TArray<uint8>& InBytes)
+int32 FRSA::EncryptPublic(const TArrayView<const uint8> InSource, TArray<uint8>& OutDestination, const FRSAKeyHandle InKey)
 {
-	int32 Count = InBytes.Num();
-
-	while (Count > 0 && InBytes[Count - 1] == 0)
+	if (IEngineCrypto* EngineCrypto = GetEngineCrypto())
 	{
-		Count--;
+		return EngineCrypto->EncryptPublic(InSource, OutDestination, InKey);
 	}
 
-	return Count;
+	return -1;
 }
 
-FRSA::TKeyPtr FRSA::CreateKey(const TArray<uint8>& InPublicExponent, const TArray<uint8>& InPrivateExponent, const TArray<uint8>& InModulus)
+int32 FRSA::EncryptPrivate(const TArrayView<const uint8> InSource, TArray<uint8>& OutDestination, const FRSAKeyHandle InKey)
 {
-	// The key data generated by openSSL is little endian, which is great for the latest version of OpenSSL and for our TBigInt. The older
-	// version of OpenSSL still in use by Linux/Max doesn't have the same little-endian import functions so we have to swizzle the data 
-	// in that case. OpenSSL should be upgraded on those platforms at some point, at which point we can remove this sub-optimal behavior
-	int32 NumPublicExponentBytes = NumElementsIgnoringTrailingZeroes(InPublicExponent);
-	int32 NumPrivateExponentBytes = NumElementsIgnoringTrailingZeroes(InPrivateExponent);
-	int32 NumModulusBytes = NumElementsIgnoringTrailingZeroes(InModulus);
-
-	int32 RequiredNumBytes = FMath::Max(NumModulusBytes, FMath::Max(NumPublicExponentBytes, NumPrivateExponentBytes));
-	RequiredNumBytes = FMath::RoundUpToPowerOfTwo(RequiredNumBytes);
-	int32 RequiredNumBits = RequiredNumBytes * 8;
-
-	// With the legacy fixed key system, which is based on a templated large integer class, we can't be that
-	// dynamic with the key lengths. Eventually, we'll move over to the totally dynamic version, but this is still an
-	// expansion on the previous functionality
-	FRSA::TKeyPtr Result;
-	
-	switch (RequiredNumBits)
+	if (IEngineCrypto* EngineCrypto = GetEngineCrypto())
 	{
-	case 4096: Result = MakeShared<FFixedKey<4096>, FRSA::KeyThreadMode>(InPublicExponent, InPrivateExponent, InModulus); break;
-	case 2048: Result = MakeShared<FFixedKey<2048>, FRSA::KeyThreadMode>(InPublicExponent, InPrivateExponent, InModulus); break;
-	case 1024: Result = MakeShared<FFixedKey<1024>, FRSA::KeyThreadMode>(InPublicExponent, InPrivateExponent, InModulus); break;
-	case 512: Result = MakeShared<FFixedKey<512>, FRSA::KeyThreadMode>(InPublicExponent, InPrivateExponent, InModulus); break;
-	case 256: Result = MakeShared<FFixedKey<256>, FRSA::KeyThreadMode>(InPublicExponent, InPrivateExponent, InModulus); break;
-	default: break;
+		return EngineCrypto->EncryptPrivate(InSource, OutDestination, InKey);
 	}
-		
-	return Result;
-}
 
-bool FRSA::Encrypt(FRSA::EKeyType InKeyType, const uint8* InSource, int32 InSourceSizeInBytes, TArray<uint8>& OutDestination, FRSA::TKeyPtr InKey)
-{
-	return InKey->Encrypt(InKeyType, InSource, InSourceSizeInBytes, OutDestination);
-}
-
-bool FRSA::Decrypt(FRSA::EKeyType InKeyType, const TArray<uint8>& InSource, uint8* OutDestination, int32 OutDestinationSizeInBytes, FRSA::TKeyPtr InKey)
-{
-	return InKey->Decrypt(InKeyType, InSource, OutDestination, OutDestinationSizeInBytes);
+	return -1;
 }
-#endif
 
-int32 FRSA::GetKeySizeInBits(FRSA::TKeyPtr InKey)
+int32 FRSA::DecryptPublic(const TArrayView<const uint8> InSource, TArray<uint8>& OutDestination, const FRSAKeyHandle InKey)
 {
-	return InKey->GetKeySizeInBits();
-}
-
-int32 FRSA::GetMaxDataSizeInBytes(FRSA::TKeyPtr InKey)
-{
-	return InKey->GetMaxDataSize();
-}
-
-bool FRSA::Encrypt(EKeyType InKeyType, const TArray<uint8>& InSource, TArray<uint8>& OutDestination, FRSA::TKeyPtr InKey)
-{
-	return Encrypt(InKeyType, InSource.GetData(), InSource.Num(), OutDestination, InKey);
-}
-
-bool FRSA::Decrypt(EKeyType InKeyType, const TArray<uint8>& InSource, TArray<uint8>& OutDestination, FRSA::TKeyPtr InKey)
-{
-	OutDestination.SetNum(InKey->GetKeySizeInBytes());
-	return Decrypt(InKeyType, InSource, OutDestination.GetData(), OutDestination.Num(), InKey);
-}
-
-#if RSA_USE_OPENSSL && USE_LEGACY_OPENSSL
-/**
- * Module interface for RSA.
- */
-class FRSAModule : public FDefaultModuleImpl
-{
-public:
-
-	virtual void StartupModule() override
+	if (IEngineCrypto* EngineCrypto = GetEngineCrypto())
 	{
-		// Provide a locking callback so that an RSA key can be used from multiple threads safely.
-		GetLocks().AddDefaulted(CRYPTO_num_locks());
-		CRYPTO_set_locking_callback(LockingCallback);
+		return EngineCrypto->DecryptPublic(InSource, OutDestination, InKey);
 	}
 
-	virtual void ShutdownModule() override
-	{
-		CRYPTO_set_locking_callback(nullptr);
-
-		// Just in case
-		for (FCriticalSection& CriticalSection : GetLocks())
-		{
-			CriticalSection.Unlock();
-		}
-	}
+	return -1;
+}
 
-	static TArray<FCriticalSection>& GetLocks()
+int32 FRSA::DecryptPrivate(const TArrayView<const uint8> InSource, TArray<uint8>& OutDestination, const FRSAKeyHandle InKey)
+{
+	if (IEngineCrypto* EngineCrypto = GetEngineCrypto())
 	{
-		static TArray<FCriticalSection> Locks;
-		return Locks;
+		return EngineCrypto->DecryptPrivate(InSource, OutDestination, InKey);
 	}
 
-	static void LockingCallback(int mode, int n, const char *file, int line)
-	{
-		TArray<FCriticalSection>& Locks = GetLocks();
-		check(Locks.Num() > 0);
-
-		if (mode & CRYPTO_LOCK)
-			Locks[n].Lock();
-		else
-			Locks[n].Unlock();
-	}
-};
+	return -1;
+}
 
-IMPLEMENT_MODULE(FRSAModule, RSA);
-#else
-// OpenSSL 1.1.1+ handles thread safety of keys internally, so there is no need to provide any custom handling and we can just use a default module implementation
-IMPLEMENT_MODULE(FDefaultModuleImpl, RSA);
-#endif
\ No newline at end of file
+IMPLEMENT_MODULE(FDefaultModuleImpl, RSA);
\ No newline at end of file
diff --git a/Engine/Source/Runtime/RSA/Public/RSA.h b/Engine/Source/Runtime/RSA/Public/RSA.h
--- a/Engine/Source/Runtime/RSA/Public/RSA.h
+++ b/Engine/Source/Runtime/RSA/Public/RSA.h
@@ -3,47 +3,48 @@
 
 #include "CoreTypes.h"
 #include "Templates/SharedPointer.h"
+#include "Misc/IEngineCrypto.h"
 
 struct RSA_API FRSA
 {
-	struct FKey;
-	static const ESPMode KeyThreadMode = ESPMode::ThreadSafe;
-	typedef TSharedPtr<FKey, KeyThreadMode> TKeyPtr;
-
 	/**
-	 * Enumerate which key to use when performing encrypt/decrpt operations
+	 * Create a new RSA public/private key from the supplied exponents and modulus binary data. Each of these arrays should contain a single little endian
+	 * large integer value
 	 */
-	enum class EKeyType
-	{
-		Public,
-		Private,
-	};
+	static FRSAKeyHandle CreateKey(const TArray<uint8>& InPublicExponent, const TArray<uint8>& InPrivateExponent, const TArray<uint8>& InModulus);
 
 	/**
-	 * Create a new RSA public/private key from the supplied exponents and modulus binary data. Each of these arrays should contain a single little endian
-	 * large integer value
+	 * Destroy the supplied key
 	 */
-	static TKeyPtr CreateKey(const TArray<uint8>& InPublicExponent, const TArray<uint8>& InPrivateExponent, const TArray<uint8>& InModulus);
+	static void DestroyKey(const FRSAKeyHandle InKey);
 
 	/**
-	 * Returns the size in bits of the supplied key
+	 * Returns the size in bytes of the supplied key
 	 */
-	static int32 GetKeySizeInBits(TKeyPtr InKey);
+	static int32 GetKeySize(const FRSAKeyHandle InKey);
 
 	/**
 	 * Returns the maximum number of bytes that can be encrypted in a single payload
 	 */
-	static int32 GetMaxDataSizeInBytes(TKeyPtr InKey);
+	static int32 GetMaxDataSize(const FRSAKeyHandle InKey);
+
+	/**
+	 * Encrypt the supplied byte data using the given public key
+	 */
+	static int32 EncryptPublic(const TArrayView<const uint8> InSource, TArray<uint8>& OutDestination, const FRSAKeyHandle InKey);
+
+	/**
+	 * Encrypt the supplied byte data using the given private key
+	 */
+	static int32 EncryptPrivate(const TArrayView<const uint8> InSource, TArray<uint8>& OutDestination, const FRSAKeyHandle InKey);
 
 	/**
-	 * Encrypt the supplied byte data using the given key
+	 * Decrypt the supplied byte data using the given public key
 	 */
-	static bool Encrypt(EKeyType InKeyType, const uint8* InSource, int32 InSourceSizeInBytes, TArray<uint8>& OutDestination, TKeyPtr InKey);
-	static bool Encrypt(EKeyType InKeyType, const TArray<uint8>& InSource, TArray<uint8>& OutDestination, TKeyPtr InKey);
+	static int32 DecryptPublic(const TArrayView<const uint8> InSource, TArray<uint8>& OutDestination, const FRSAKeyHandle InKey);
 
 	/**
-	 * Decrypt the supplied byte data using the given key
+	 * Encrypt the supplied byte data using the given private key
 	 */
-	static bool Decrypt(EKeyType InKeyType, const TArray<uint8>& InSource, uint8* OutDestination, int32 OutDestinationSizeInBytes, TKeyPtr InKey);
-	static bool Decrypt(EKeyType InKeyType, const TArray<uint8>& InSource, TArray<uint8>& OutDestination, TKeyPtr InKey);
+	static int32 DecryptPrivate(const TArrayView<const uint8> InSource, TArray<uint8>& OutDestination, const FRSAKeyHandle InKey);
 };
\ No newline at end of file
diff --git a/Engine/Source/Runtime/RSA/RSA.Build.cs b/Engine/Source/Runtime/RSA/RSA.Build.cs
--- a/Engine/Source/Runtime/RSA/RSA.Build.cs
+++ b/Engine/Source/Runtime/RSA/RSA.Build.cs
@@ -7,23 +7,5 @@ public class RSA : ModuleRules
 	public RSA(ReadOnlyTargetRules Target) : base(Target)
 	{
 		PrivateDependencyModuleNames.Add("Core");
-
-        switch (Target.Platform)
-        {
-			case UnrealTargetPlatform.Win64:
-			case UnrealTargetPlatform.Win32:
-			case UnrealTargetPlatform.Mac:
-			case UnrealTargetPlatform.Linux:
-				{
-					PrivateDependencyModuleNames.Add("OpenSSL");
-					PrivateDefinitions.Add("RSA_USE_OPENSSL=1");
-					break;
-				}
-			default:
-				{
-					PrivateDefinitions.Add("RSA_USE_OPENSSL=0");
-					break;
-				}
-		}
 	}
 }
